Fix crash when clicking under the pipelines in the "interchange pipeline configuration" dialog
#rnx

diff --git a/Engine/Plugins/Experimental/Interchange/Editor/Source/Pipelines/Private/SInterchangePipelineConfigurationDialog.cpp b/Engine/Plugins/Experimental/Interchange/Editor/Source/Pipelines/Private/SInterchangePipelineConfigurationDialog.cpp
--- a/Engine/Plugins/Experimental/Interchange/Editor/Source/Pipelines/Private/SInterchangePipelineConfigurationDialog.cpp
+++ b/Engine/Plugins/Experimental/Interchange/Editor/Source/Pipelines/Private/SInterchangePipelineConfigurationDialog.cpp
@@ -427,8 +427,9 @@ void SInterchangePipelineConfigurationDialog::Construct(const FArguments& InArgs
 
 void SInterchangePipelineConfigurationDialog::OnSelectionChanged(TSharedPtr<FInterchangePipelineStacksTreeNodeItem> Item, ESelectInfo::Type SelectionType)
 {
+	UInterchangePipelineBase* Pipeline = Item ? Item->Pipeline : nullptr;
 	//Change the object point by the InspectorBox
-	PipelineConfigurationDetailsView->SetObject(Item->Pipeline);
+	PipelineConfigurationDetailsView->SetObject(Pipeline);
 }
 
 void SInterchangePipelineConfigurationDialog::RecursiveSavePipelineSettings(const TSharedPtr<FInterchangePipelineStacksTreeNodeItem>& ParentNode, const int32 PipelineIndex) const
