Interchange inspector: Applied different colors to the node icons based on the node container types.

diff --git a/Engine/Plugins/Experimental/Interchange/Editor/Source/Pipelines/Private/SInterchangeGraphInspectorWindow.cpp b/Engine/Plugins/Experimental/Interchange/Editor/Source/Pipelines/Private/SInterchangeGraphInspectorWindow.cpp
--- a/Engine/Plugins/Experimental/Interchange/Editor/Source/Pipelines/Private/SInterchangeGraphInspectorWindow.cpp
+++ b/Engine/Plugins/Experimental/Interchange/Editor/Source/Pipelines/Private/SInterchangeGraphInspectorWindow.cpp
@@ -11,6 +11,7 @@
 #include "Nodes/InterchangeBaseNodeContainer.h"
 #include "PropertyEditorModule.h"
 #include "Styling/SlateIconFinder.h"
+#include "Styling/StyleColors.h"
 #include "Widgets/Images/SImage.h"
 #include "Widgets/Input/SButton.h"
 #include "Widgets/Input/SCheckBox.h"
@@ -111,6 +112,25 @@ public:
 			Tooltip += TEXT(" [") + InterchangeNode->GetObjectClass()->GetName() + TEXT("]");
 		}
 
+		FSlateColor TypeIconColor;
+
+		switch(InterchangeNode->GetNodeContainerType())
+		{
+		case EInterchangeNodeContainerType::TranslatedAsset:
+			TypeIconColor = FStyleColors::AccentBlue;
+			break;
+		case EInterchangeNodeContainerType::TranslatedScene:
+			TypeIconColor = FStyleColors::AccentGreen;
+			break;
+		case EInterchangeNodeContainerType::FactoryData:
+			TypeIconColor = FStyleColors::AccentPurple;
+			break;
+		case EInterchangeNodeContainerType::None:
+		default:
+			TypeIconColor = FStyleColors::AccentRed;
+			break;
+		}
+
 		this->ChildSlot
 		[
 			SNew(SHorizontalBox)
@@ -133,6 +153,7 @@ public:
 			[
 				SNew(SImage)
 				.Image(TypeIcon)
+				.ColorAndOpacity(TypeIconColor)
 				.Visibility(TypeIcon != FEditorStyle::GetDefaultBrush() ? EVisibility::Visible : EVisibility::Collapsed)
 			]
 			+ SHorizontalBox::Slot()
