Make sure all 3f/3d type are supported correctly
- Remove from the attribute storage the type that represent 3f or 3d (i.e FVector, FTransform).
- Make sure the detail viewer for the interchange node support the specific 3f/3d type

#rnx

diff --git a/Engine/Plugins/Experimental/Interchange/Editor/Source/Pipelines/Private/InterchangeEditorPipelineDetails.cpp b/Engine/Plugins/Experimental/Interchange/Editor/Source/Pipelines/Private/InterchangeEditorPipelineDetails.cpp
--- a/Engine/Plugins/Experimental/Interchange/Editor/Source/Pipelines/Private/InterchangeEditorPipelineDetails.cpp
+++ b/Engine/Plugins/Experimental/Interchange/Editor/Source/Pipelines/Private/InterchangeEditorPipelineDetails.cpp
@@ -172,15 +172,39 @@ void FInterchangeBaseNodeDetailsCustomization::AddAttributeRow(UE::Interchange::
 		}
 		break;
 
-		case UE::Interchange::EAttributeTypes::Transform:
+		case UE::Interchange::EAttributeTypes::Transform3f:
 		{
-			BuildTransformValueContent(AttributeCategory, AttributeKey);
+			BuildTransformValueContent<FTransform3f, FVector3f, FQuat4f, float>(AttributeCategory, AttributeKey);
 		}
 		break;
 
-		case UE::Interchange::EAttributeTypes::Box:
+		case UE::Interchange::EAttributeTypes::Transform3d:
 		{
-			BuildBoxValueContent(AttributeCategory, AttributeKey);
+			BuildTransformValueContent<FTransform3d, FVector3d, FQuat4d, double>(AttributeCategory, AttributeKey);
+		}
+		break;
+
+		case UE::Interchange::EAttributeTypes::Box3f:
+		{
+			BuildBoxValueContent<FBox3f, FVector3f, float>(AttributeCategory, AttributeKey);
+		}
+		break;
+
+		case UE::Interchange::EAttributeTypes::Box3d:
+		{
+			BuildBoxValueContent<FBox3d, FVector3d, double>(AttributeCategory, AttributeKey);
+		}
+		break;
+
+		case UE::Interchange::EAttributeTypes::Vector3f:
+		{
+			BuildVectorValueContent<FVector3f, float>(AttributeCategory, AttributeKey);
+		}
+		break;
+
+		case UE::Interchange::EAttributeTypes::Vector3d:
+		{
+			BuildVectorValueContent<FVector3d, double>(AttributeCategory, AttributeKey);
 		}
 		break;
 
@@ -390,14 +414,13 @@ void FInterchangeBaseNodeDetailsCustomization::BuildStringValueContent(IDetailCa
         ]
     ];
 }
-
+template<typename TransformnType, typename VectorType, typename QuatType, typename NumericType>
 void FInterchangeBaseNodeDetailsCustomization::BuildTransformValueContent(IDetailCategoryBuilder& AttributeCategory, UE::Interchange::FAttributeKey& AttributeKey)
 {
 	UE::Interchange::EAttributeTypes AttributeType = InterchangeBaseNode->GetAttributeType(AttributeKey);
-	check(AttributeType == UE::Interchange::EAttributeTypes::Transform);
 
 	{
-		const UE::Interchange::FAttributeStorage::TAttributeHandle<FTransform> AttributeHandle = InterchangeBaseNode->GetAttributeHandle<FTransform>(AttributeKey);
+		const UE::Interchange::FAttributeStorage::TAttributeHandle<TransformnType> AttributeHandle = InterchangeBaseNode->GetAttributeHandle<TransformnType>(AttributeKey);
 		if (!AttributeHandle.IsValid())
 		{
 			CreateInvalidHandleRow(AttributeCategory, AttributeKey);
@@ -414,10 +437,10 @@ void FInterchangeBaseNodeDetailsCustomization::BuildTransformValueContent(IDetai
 		CreateNameWidget(AttributeKey)
 	];
 
-	auto GetRotationValue = [](UInterchangeBaseNode* BaseNode, UE::Interchange::FAttributeKey& Key)->FQuat
+	auto GetRotationValue = [](UInterchangeBaseNode* BaseNode, UE::Interchange::FAttributeKey& Key)->QuatType
 	{
-		FTransform TransformValue;
-		const UE::Interchange::FAttributeStorage::TAttributeHandle<FTransform> AttributeHandle = BaseNode->GetAttributeHandle<FTransform>(Key);
+		TransformnType TransformValue;
+		const UE::Interchange::FAttributeStorage::TAttributeHandle<TransformnType> AttributeHandle = BaseNode->GetAttributeHandle<TransformnType>(Key);
 		if (AttributeHandle.IsValid())
 		{
 			AttributeHandle.Get(TransformValue);
@@ -425,10 +448,10 @@ void FInterchangeBaseNodeDetailsCustomization::BuildTransformValueContent(IDetai
 		return TransformValue.GetRotation();
 	};
 
-	auto SetRotationValue = [](UInterchangeBaseNode* BaseNode, UE::Interchange::FAttributeKey& Key, const FQuat& Value)
+	auto SetRotationValue = [](UInterchangeBaseNode* BaseNode, UE::Interchange::FAttributeKey& Key, const QuatType& Value)
 	{
-		FTransform TransformValue;
-		UE::Interchange::FAttributeStorage::TAttributeHandle<FTransform> AttributeHandle = BaseNode->GetAttributeHandle<FTransform>(Key);
+		TransformnType TransformValue;
+		UE::Interchange::FAttributeStorage::TAttributeHandle<TransformnType> AttributeHandle = BaseNode->GetAttributeHandle<TransformnType>(Key);
 		if (AttributeHandle.IsValid())
 		{
 			AttributeHandle.Get(TransformValue);
@@ -437,10 +460,10 @@ void FInterchangeBaseNodeDetailsCustomization::BuildTransformValueContent(IDetai
 		}
 	};
 
-	auto GetTranslationValue = [](UInterchangeBaseNode* BaseNode, UE::Interchange::FAttributeKey& Key)->FVector
+	auto GetTranslationValue = [](UInterchangeBaseNode* BaseNode, UE::Interchange::FAttributeKey& Key)->VectorType
 	{
-		FTransform TransformValue;
-		const UE::Interchange::FAttributeStorage::TAttributeHandle<FTransform> AttributeHandle = BaseNode->GetAttributeHandle<FTransform>(Key);
+		TransformnType TransformValue;
+		const UE::Interchange::FAttributeStorage::TAttributeHandle<TransformnType> AttributeHandle = BaseNode->GetAttributeHandle<TransformnType>(Key);
 		if (AttributeHandle.IsValid())
 		{
 			AttributeHandle.Get(TransformValue);
@@ -448,10 +471,10 @@ void FInterchangeBaseNodeDetailsCustomization::BuildTransformValueContent(IDetai
 		return TransformValue.GetTranslation();
 	};
 
-	auto SetTranslationValue = [](UInterchangeBaseNode* BaseNode, UE::Interchange::FAttributeKey& Key, const FVector& Value)
+	auto SetTranslationValue = [](UInterchangeBaseNode* BaseNode, UE::Interchange::FAttributeKey& Key, const VectorType& Value)
 	{
-		FTransform TransformValue;
-		UE::Interchange::FAttributeStorage::TAttributeHandle<FTransform> AttributeHandle = BaseNode->GetAttributeHandle<FTransform>(Key);
+		TransformnType TransformValue;
+		UE::Interchange::FAttributeStorage::TAttributeHandle<TransformnType> AttributeHandle = BaseNode->GetAttributeHandle<TransformnType>(Key);
 		if (AttributeHandle.IsValid())
 		{
 			AttributeHandle.Get(TransformValue);
@@ -460,10 +483,10 @@ void FInterchangeBaseNodeDetailsCustomization::BuildTransformValueContent(IDetai
 		}
 	};
 
-	auto GetScale3DValue = [](UInterchangeBaseNode* BaseNode, UE::Interchange::FAttributeKey& Key)->FVector
+	auto GetScale3DValue = [](UInterchangeBaseNode* BaseNode, UE::Interchange::FAttributeKey& Key)->VectorType
 	{
-		FTransform TransformValue;
-		const UE::Interchange::FAttributeStorage::TAttributeHandle<FTransform> AttributeHandle = BaseNode->GetAttributeHandle<FTransform>(Key);
+		TransformnType TransformValue;
+		const UE::Interchange::FAttributeStorage::TAttributeHandle<TransformnType> AttributeHandle = BaseNode->GetAttributeHandle<TransformnType>(Key);
 		if (AttributeHandle.IsValid())
 		{
 			AttributeHandle.Get(TransformValue);
@@ -471,10 +494,10 @@ void FInterchangeBaseNodeDetailsCustomization::BuildTransformValueContent(IDetai
 		return TransformValue.GetScale3D();
 	};
 
-	auto SetScale3DValue = [](UInterchangeBaseNode* BaseNode, UE::Interchange::FAttributeKey& Key, const FVector& Value)
+	auto SetScale3DValue = [](UInterchangeBaseNode* BaseNode, UE::Interchange::FAttributeKey& Key, const VectorType& Value)
 	{
-		FTransform TransformValue;
-		UE::Interchange::FAttributeStorage::TAttributeHandle<FTransform> AttributeHandle = BaseNode->GetAttributeHandle<FTransform>(Key);
+		TransformnType TransformValue;
+		UE::Interchange::FAttributeStorage::TAttributeHandle<TransformnType> AttributeHandle = BaseNode->GetAttributeHandle<TransformnType>(Key);
 		if (AttributeHandle.IsValid())
 		{
 			AttributeHandle.Get(TransformValue);
@@ -491,7 +514,7 @@ void FInterchangeBaseNodeDetailsCustomization::BuildTransformValueContent(IDetai
     ]
     .ValueContent()
     [
-        CreateVectorWidget(GetTranslationValue, SetTranslationValue, AttributeKey)
+        CreateVectorWidget<VectorType, NumericType>(GetTranslationValue, SetTranslationValue, AttributeKey)
     ];
 
 	const FString RotationName = TEXT("Rotation");
@@ -502,7 +525,7 @@ void FInterchangeBaseNodeDetailsCustomization::BuildTransformValueContent(IDetai
 	]
 	.ValueContent()
 	[
-		CreateQuaternionWidget(GetRotationValue, SetRotationValue, AttributeKey)
+		CreateQuaternionWidget<QuatType, NumericType>(GetRotationValue, SetRotationValue, AttributeKey)
 	];
 
 	const FString Scale3DName = TEXT("Scale3D");
@@ -513,7 +536,7 @@ void FInterchangeBaseNodeDetailsCustomization::BuildTransformValueContent(IDetai
     ]
     .ValueContent()
     [
-        CreateVectorWidget(GetScale3DValue, SetScale3DValue, AttributeKey)
+        CreateVectorWidget<VectorType, NumericType>(GetScale3DValue, SetScale3DValue, AttributeKey)
     ];
 }
 
@@ -747,13 +770,12 @@ void FInterchangeBaseNodeDetailsCustomization::InternalBuildColorValueContent(ID
     ];
 }
 
+template<typename BoxType, typename VectorType, typename NumericType>
 void FInterchangeBaseNodeDetailsCustomization::BuildBoxValueContent(IDetailCategoryBuilder& AttributeCategory, UE::Interchange::FAttributeKey& AttributeKey)
 {
 	UE::Interchange::EAttributeTypes AttributeType = InterchangeBaseNode->GetAttributeType(AttributeKey);
-	check(AttributeType == UE::Interchange::EAttributeTypes::Box);
-
 	{
-		const UE::Interchange::FAttributeStorage::TAttributeHandle<FBox> AttributeHandle = InterchangeBaseNode->GetAttributeHandle<FBox>(AttributeKey);
+		const UE::Interchange::FAttributeStorage::TAttributeHandle<BoxType> AttributeHandle = InterchangeBaseNode->GetAttributeHandle<BoxType>(AttributeKey);
 		if (!AttributeHandle.IsValid())
 		{
 			CreateInvalidHandleRow(AttributeCategory, AttributeKey);
@@ -770,10 +792,10 @@ void FInterchangeBaseNodeDetailsCustomization::BuildBoxValueContent(IDetailCateg
 		CreateNameWidget(AttributeKey)
 	];
 
-	auto GetMinimumValue = [](UInterchangeBaseNode* BaseNode, UE::Interchange::FAttributeKey& Key)->FVector
+	auto GetMinimumValue = [](UInterchangeBaseNode* BaseNode, UE::Interchange::FAttributeKey& Key)->VectorType
 	{
-		FBox BoxValue;
-		const UE::Interchange::FAttributeStorage::TAttributeHandle<FBox> AttributeHandle = BaseNode->GetAttributeHandle<FBox>(Key);
+		BoxType BoxValue;
+		const UE::Interchange::FAttributeStorage::TAttributeHandle<BoxType> AttributeHandle = BaseNode->GetAttributeHandle<BoxType>(Key);
 		if (AttributeHandle.IsValid())
 		{
 			AttributeHandle.Get(BoxValue);
@@ -781,10 +803,10 @@ void FInterchangeBaseNodeDetailsCustomization::BuildBoxValueContent(IDetailCateg
 		return BoxValue.Min;
 	};
 
-	auto SetMinimumValue = [](UInterchangeBaseNode* BaseNode, UE::Interchange::FAttributeKey& Key, const FVector& Value)
+	auto SetMinimumValue = [](UInterchangeBaseNode* BaseNode, UE::Interchange::FAttributeKey& Key, const VectorType& Value)
 	{
-		FBox BoxValue;
-		UE::Interchange::FAttributeStorage::TAttributeHandle<FBox> AttributeHandle = BaseNode->GetAttributeHandle<FBox>(Key);
+		BoxType BoxValue;
+		UE::Interchange::FAttributeStorage::TAttributeHandle<BoxType> AttributeHandle = BaseNode->GetAttributeHandle<BoxType>(Key);
 		if (AttributeHandle.IsValid())
 		{
 			AttributeHandle.Get(BoxValue);
@@ -793,10 +815,10 @@ void FInterchangeBaseNodeDetailsCustomization::BuildBoxValueContent(IDetailCateg
 		}
 	};
 
-	auto GetMaximumValue = [](UInterchangeBaseNode* BaseNode, UE::Interchange::FAttributeKey& Key)->FVector
+	auto GetMaximumValue = [](UInterchangeBaseNode* BaseNode, UE::Interchange::FAttributeKey& Key)->VectorType
 	{
-		FBox BoxValue;
-		const UE::Interchange::FAttributeStorage::TAttributeHandle<FBox> AttributeHandle = BaseNode->GetAttributeHandle<FBox>(Key);
+		BoxType BoxValue;
+		const UE::Interchange::FAttributeStorage::TAttributeHandle<BoxType> AttributeHandle = BaseNode->GetAttributeHandle<BoxType>(Key);
 		if (AttributeHandle.IsValid())
 		{
 			AttributeHandle.Get(BoxValue);
@@ -804,10 +826,10 @@ void FInterchangeBaseNodeDetailsCustomization::BuildBoxValueContent(IDetailCateg
 		return BoxValue.Max;
 	};
 
-	auto SetMaximumValue = [](UInterchangeBaseNode* BaseNode, UE::Interchange::FAttributeKey& Key, const FVector& Value)
+	auto SetMaximumValue = [](UInterchangeBaseNode* BaseNode, UE::Interchange::FAttributeKey& Key, const VectorType& Value)
 	{
-		FBox BoxValue;
-		UE::Interchange::FAttributeStorage::TAttributeHandle<FBox> AttributeHandle = BaseNode->GetAttributeHandle<FBox>(Key);
+		BoxType BoxValue;
+		UE::Interchange::FAttributeStorage::TAttributeHandle<BoxType> AttributeHandle = BaseNode->GetAttributeHandle<BoxType>(Key);
 		if (AttributeHandle.IsValid())
 		{
 			AttributeHandle.Get(BoxValue);
@@ -824,7 +846,7 @@ void FInterchangeBaseNodeDetailsCustomization::BuildBoxValueContent(IDetailCateg
 	]
 	.ValueContent()
 	[
-		CreateVectorWidget(GetMinimumValue, SetMinimumValue, AttributeKey)
+		CreateVectorWidget<VectorType, NumericType>(GetMinimumValue, SetMinimumValue, AttributeKey)
 	];
 
 	const FString MaximumVectorName = TEXT("Maximum");
@@ -835,7 +857,53 @@ void FInterchangeBaseNodeDetailsCustomization::BuildBoxValueContent(IDetailCateg
 	]
 	.ValueContent()
 	[
-		CreateVectorWidget(GetMaximumValue, SetMaximumValue, AttributeKey)
+		CreateVectorWidget<VectorType, NumericType>(GetMaximumValue, SetMaximumValue, AttributeKey)
+	];
+}
+
+template<typename VectorType, typename NumericType>
+void FInterchangeBaseNodeDetailsCustomization::BuildVectorValueContent(IDetailCategoryBuilder& AttributeCategory, UE::Interchange::FAttributeKey& AttributeKey)
+{
+	UE::Interchange::EAttributeTypes AttributeType = InterchangeBaseNode->GetAttributeType(AttributeKey);
+
+	{
+		const UE::Interchange::FAttributeStorage::TAttributeHandle<VectorType> AttributeHandle = InterchangeBaseNode->GetAttributeHandle<VectorType>(AttributeKey);
+		if (!AttributeHandle.IsValid())
+		{
+			CreateInvalidHandleRow(AttributeCategory, AttributeKey);
+			return;
+		}
+	}
+
+	auto GetValue = [](UInterchangeBaseNode* BaseNode, UE::Interchange::FAttributeKey& Key)->VectorType
+	{
+		VectorType VectorValue;
+		const UE::Interchange::FAttributeStorage::TAttributeHandle<VectorType> AttributeHandle = BaseNode->GetAttributeHandle<VectorType>(Key);
+		if (AttributeHandle.IsValid())
+		{
+			AttributeHandle.Get(VectorValue);
+		}
+		return VectorValue;
+	};
+
+	auto SetValue = [](UInterchangeBaseNode* BaseNode, UE::Interchange::FAttributeKey& Key, const VectorType& Value)
+	{
+		UE::Interchange::FAttributeStorage::TAttributeHandle<VectorType> AttributeHandle = BaseNode->GetAttributeHandle<VectorType>(Key);
+		if (AttributeHandle.IsValid())
+		{
+			AttributeHandle.Set(Value);
+		}
+	};
+
+	const bool bAdvancedProperty = false;
+	FDetailWidgetRow& VectorRow = AttributeCategory.AddCustomRow(FText::FromString(AttributeKey.Key), bAdvancedProperty);
+	VectorRow.NameContent()
+	[
+		CreateNameWidget(AttributeKey)
+	]
+	.ValueContent()
+	[
+		CreateVectorWidget<VectorType, NumericType>(GetValue, SetValue, AttributeKey)
 	];
 }
 
@@ -875,18 +943,18 @@ TSharedRef<SWidget> FInterchangeBaseNodeDetailsCustomization::CreateSimpleNameWi
 		.Font(IDetailLayoutBuilder::GetDetailFont());
 }
 
-template<typename FunctorGet, typename FunctorSet>
+template<typename VectorType, typename NumericWidgetType, typename FunctorGet, typename FunctorSet>
 TSharedRef<SWidget> FInterchangeBaseNodeDetailsCustomization::CreateVectorWidget(FunctorGet GetValue, FunctorSet SetValue, UE::Interchange::FAttributeKey& AttributeKey)
 {
-	auto GetComponentValue = [&GetValue](int32 ComponentIndex, UInterchangeBaseNode* BaseNode, UE::Interchange::FAttributeKey& Key)->float
+	auto GetComponentValue = [&GetValue](int32 ComponentIndex, UInterchangeBaseNode* BaseNode, UE::Interchange::FAttributeKey& Key)->NumericWidgetType
 	{
-		FVector Value = GetValue(BaseNode, Key);
+		VectorType Value = GetValue(BaseNode, Key);
 		return Value[ComponentIndex];
 	};
 
-	auto SetComponentValue = [&GetValue, &SetValue](int32 ComponentIndex, UInterchangeBaseNode* BaseNode, float ComponentValue, UE::Interchange::FAttributeKey& Key)
+	auto SetComponentValue = [&GetValue, &SetValue](int32 ComponentIndex, UInterchangeBaseNode* BaseNode, NumericWidgetType ComponentValue, UE::Interchange::FAttributeKey& Key)
 	{
-		FVector Value = GetValue(BaseNode, Key);
+		VectorType Value = GetValue(BaseNode, Key);
 		Value[ComponentIndex] = ComponentValue;
 		SetValue(BaseNode, Key, Value);
     };
@@ -896,26 +964,26 @@ TSharedRef<SWidget> FInterchangeBaseNodeDetailsCustomization::CreateVectorWidget
         + SHorizontalBox::Slot()
         .AutoWidth()
         [
-			MakeNumericWidget<float>(0, GetComponentValue, SetComponentValue, AttributeKey)
+			MakeNumericWidget<NumericWidgetType>(0, GetComponentValue, SetComponentValue, AttributeKey)
         ]
 		+ SHorizontalBox::Slot()
 	    .AutoWidth()
 	    [
-	        MakeNumericWidget<float>(1, GetComponentValue, SetComponentValue, AttributeKey)
+	        MakeNumericWidget<NumericWidgetType>(1, GetComponentValue, SetComponentValue, AttributeKey)
 	    ]
 		+ SHorizontalBox::Slot()
 	    .AutoWidth()
 	    [
-	        MakeNumericWidget<float>(2, GetComponentValue, SetComponentValue, AttributeKey)
+	        MakeNumericWidget<NumericWidgetType>(2, GetComponentValue, SetComponentValue, AttributeKey)
 	    ];
 }
 
-template<typename FunctorGet, typename FunctorSet>
+template<typename QuatType, typename NumericWidgetType, typename FunctorGet, typename FunctorSet>
 TSharedRef<SWidget> FInterchangeBaseNodeDetailsCustomization::CreateQuaternionWidget(FunctorGet GetValue, FunctorSet SetValue, UE::Interchange::FAttributeKey& AttributeKey)
 {
-	auto GetComponentValue = [&GetValue](int32 ComponentIndex, UInterchangeBaseNode* BaseNode, UE::Interchange::FAttributeKey& Key)->float
+	auto GetComponentValue = [&GetValue](int32 ComponentIndex, UInterchangeBaseNode* BaseNode, UE::Interchange::FAttributeKey& Key)->NumericWidgetType
 	{
-		FQuat Value = GetValue(BaseNode, Key);
+		QuatType Value = GetValue(BaseNode, Key);
 		if(ComponentIndex == 0)
 		{
 			return Value.X;
@@ -938,9 +1006,9 @@ TSharedRef<SWidget> FInterchangeBaseNodeDetailsCustomization::CreateQuaternionWi
 		return 0.0f;	
 	};
 
-	auto SetComponentValue = [&GetValue, &SetValue](int32 ComponentIndex, UInterchangeBaseNode* BaseNode, float ComponentValue, UE::Interchange::FAttributeKey& Key)
+	auto SetComponentValue = [&GetValue, &SetValue](int32 ComponentIndex, UInterchangeBaseNode* BaseNode, NumericWidgetType ComponentValue, UE::Interchange::FAttributeKey& Key)
 	{
-		FQuat Value = GetValue(BaseNode, Key);
+		QuatType Value = GetValue(BaseNode, Key);
 		if(ComponentIndex == 0)
 		{
 			Value.X = ComponentValue;
@@ -969,22 +1037,22 @@ TSharedRef<SWidget> FInterchangeBaseNodeDetailsCustomization::CreateQuaternionWi
         + SHorizontalBox::Slot()
         .AutoWidth()
         [
-            MakeNumericWidget<float>(0, GetComponentValue, SetComponentValue, AttributeKey)
+            MakeNumericWidget<NumericWidgetType>(0, GetComponentValue, SetComponentValue, AttributeKey)
         ]
         + SHorizontalBox::Slot()
         .AutoWidth()
         [
-            MakeNumericWidget<float>(1, GetComponentValue, SetComponentValue, AttributeKey)
+            MakeNumericWidget<NumericWidgetType>(1, GetComponentValue, SetComponentValue, AttributeKey)
         ]
         + SHorizontalBox::Slot()
         .AutoWidth()
         [
-            MakeNumericWidget<float>(2, GetComponentValue, SetComponentValue, AttributeKey)
+            MakeNumericWidget<NumericWidgetType>(2, GetComponentValue, SetComponentValue, AttributeKey)
         ]
 		+ SHorizontalBox::Slot()
 	    .AutoWidth()
 	    [
-	        MakeNumericWidget<float>(3, GetComponentValue, SetComponentValue, AttributeKey)
+	        MakeNumericWidget<NumericWidgetType>(3, GetComponentValue, SetComponentValue, AttributeKey)
 	    ];
 }
 
diff --git a/Engine/Plugins/Experimental/Interchange/Editor/Source/Pipelines/Private/InterchangeEditorPipelineDetails.h b/Engine/Plugins/Experimental/Interchange/Editor/Source/Pipelines/Private/InterchangeEditorPipelineDetails.h
--- a/Engine/Plugins/Experimental/Interchange/Editor/Source/Pipelines/Private/InterchangeEditorPipelineDetails.h
+++ b/Engine/Plugins/Experimental/Interchange/Editor/Source/Pipelines/Private/InterchangeEditorPipelineDetails.h
@@ -56,6 +56,7 @@ private:
 
 	//FTransform: Add a group to the category that contain the Translation(FVector), Rotation(FQuat) and scale3D(FVector)
 	//Any numeric field are read/write
+	template<typename TransformnType, typename VectorType, typename QuatType, typename NumericType>
 	void BuildTransformValueContent(IDetailCategoryBuilder& AttributeCategory, UE::Interchange::FAttributeKey& AttributeKey);
 
 	//FColor: Add a row to the category that contain the RGBA value
@@ -71,8 +72,14 @@ private:
 
 	//FBox: Add a group to the category that contain the Minimum corner position(FVector), and the maximum corner position(FVector)
 	//Any numeric field are read/write
+	template<typename BoxType, typename VectorType, typename NumericType>
 	void BuildBoxValueContent(IDetailCategoryBuilder& AttributeCategory, UE::Interchange::FAttributeKey& AttributeKey);
 
+	//FVector: Add a row to the category that display a FVector.
+	//Any numeric field are read/write
+	template<typename VectorType, typename NumericType>
+	void BuildVectorValueContent(IDetailCategoryBuilder& AttributeCategory, UE::Interchange::FAttributeKey& AttributeKey);
+
 	/**
 	 * AttributeType Build functions End
 	 ******************************************************************/
@@ -92,11 +99,11 @@ private:
 	TSharedRef<SWidget> CreateSimpleNameWidget(const FString& SimpleName) const;
 
 	//FVector Create a 3 digit widget (X, Y, Z)
-	template<typename FunctorGet, typename FunctorSet>
+	template<typename VectorType, typename NumericWidgetType, typename FunctorGet, typename FunctorSet>
 	TSharedRef<SWidget> CreateVectorWidget(FunctorGet GetValue, FunctorSet SetValue, UE::Interchange::FAttributeKey& AttributeKey);
 
 	//FQuat Create a 4 digit widget (X, Y, Z, W)
-	template<typename FunctorGet, typename FunctorSet>
+	template<typename QuatType, typename NumericWidgetType, typename FunctorGet, typename FunctorSet>
     TSharedRef<SWidget> CreateQuaternionWidget(FunctorGet GetValue, FunctorSet SetValue, UE::Interchange::FAttributeKey& AttributeKey);
 
 	//Create one numeric widget, support (float, double, signed/unsigned integer)
diff --git a/Engine/Plugins/Experimental/Interchange/Runtime/Source/Dispatcher/Public/InterchangeCommands.h b/Engine/Plugins/Experimental/Interchange/Runtime/Source/Dispatcher/Public/InterchangeCommands.h
--- a/Engine/Plugins/Experimental/Interchange/Runtime/Source/Dispatcher/Public/InterchangeCommands.h
+++ b/Engine/Plugins/Experimental/Interchange/Runtime/Source/Dispatcher/Public/InterchangeCommands.h
@@ -17,7 +17,7 @@ namespace UE
 		{
 		public:
 			//Major version should be updated when there is an existing API that has been change.
-			static int32 GetMajor() { return 1; }
+			static int32 GetMajor() { return 2; }
 			//Minor version must be updated when there is an addition to the API.
 			static int32 GetMinor() { return 0; }
 			//Patch version should be update if there is some bug fixes in the private code.
diff --git a/Engine/Source/Runtime/Experimental/Interchange/Core/Private/Types/AttributeStorage.cpp b/Engine/Source/Runtime/Experimental/Interchange/Core/Private/Types/AttributeStorage.cpp
--- a/Engine/Source/Runtime/Experimental/Interchange/Core/Private/Types/AttributeStorage.cpp
+++ b/Engine/Source/Runtime/Experimental/Interchange/Core/Private/Types/AttributeStorage.cpp
@@ -18,8 +18,6 @@ namespace UE
 			{
 				case EAttributeTypes::None: { AttributeTypeString = TEXT("None"); } break;
 				case EAttributeTypes::Bool: { AttributeTypeString = TEXT("Bool"); } break;
-				case EAttributeTypes::Box: { AttributeTypeString = TEXT("Box"); } break;
-				case EAttributeTypes::BoxSphereBounds: { AttributeTypeString = TEXT("BoxSphereBounds"); } break;
 				case EAttributeTypes::ByteArray: { AttributeTypeString = TEXT("ByteArray"); } break;
 				case EAttributeTypes::ByteArray64: { AttributeTypeString = TEXT("ByteArray64"); } break;
 				case EAttributeTypes::Color: { AttributeTypeString = TEXT("Color"); } break;
@@ -34,29 +32,21 @@ namespace UE
 				case EAttributeTypes::Int64: { AttributeTypeString = TEXT("Int64"); } break;
 				case EAttributeTypes::IntRect: { AttributeTypeString = TEXT("IntRect"); } break;
 				case EAttributeTypes::LinearColor: { AttributeTypeString = TEXT("LinearColor"); } break;
-				case EAttributeTypes::Matrix: { AttributeTypeString = TEXT("Matrix"); } break;
 				case EAttributeTypes::Name: { AttributeTypeString = TEXT("Name"); } break;
-				case EAttributeTypes::Plane: { AttributeTypeString = TEXT("Plane"); } break;
-				case EAttributeTypes::Quat: { AttributeTypeString = TEXT("Quat"); } break;
 				case EAttributeTypes::RandomStream: { AttributeTypeString = TEXT("RandomStream"); } break;
-				case EAttributeTypes::Rotator: { AttributeTypeString = TEXT("Rotator"); } break;
 				case EAttributeTypes::String: { AttributeTypeString = TEXT("String"); } break;
 				case EAttributeTypes::Timespan: { AttributeTypeString = TEXT("Timespan"); } break;
-				case EAttributeTypes::Transform: { AttributeTypeString = TEXT("Transform"); } break;
 				case EAttributeTypes::TwoVectors: { AttributeTypeString = TEXT("TwoVectors"); } break;
 				case EAttributeTypes::UInt8: { AttributeTypeString = TEXT("UInt8"); } break;
 				case EAttributeTypes::UInt16: { AttributeTypeString = TEXT("UInt16"); } break;
 				case EAttributeTypes::UInt32: { AttributeTypeString = TEXT("UInt32"); } break;
 				case EAttributeTypes::UInt64: { AttributeTypeString = TEXT("UInt64"); } break;
-				case EAttributeTypes::Vector: { AttributeTypeString = TEXT("Vector"); } break;
 				case EAttributeTypes::Vector2d: { AttributeTypeString = TEXT("Vector2d"); } break;
-				case EAttributeTypes::Vector4: { AttributeTypeString = TEXT("Vector4"); } break;
 				case EAttributeTypes::IntPoint: { AttributeTypeString = TEXT("IntPoint"); } break;
 				case EAttributeTypes::IntVector: { AttributeTypeString = TEXT("IntVector"); } break;
 				case EAttributeTypes::Vector2DHalf: { AttributeTypeString = TEXT("Vector2DHalf"); } break;
 				case EAttributeTypes::Float16: { AttributeTypeString = TEXT("Float16"); } break;
 				case EAttributeTypes::OrientedBox: { AttributeTypeString = TEXT("OrientedBox"); } break;
-				case EAttributeTypes::Sphere: { AttributeTypeString = TEXT("Sphere"); } break;
 				case EAttributeTypes::FrameNumber: { AttributeTypeString = TEXT("FrameNumber"); } break;
 				case EAttributeTypes::FrameRate: { AttributeTypeString = TEXT("FrameRate"); } break;
 				case EAttributeTypes::FrameTime: { AttributeTypeString = TEXT("FrameTime"); } break;
diff --git a/Engine/Source/Runtime/Experimental/Interchange/Core/Public/Types/AttributeStorage.h b/Engine/Source/Runtime/Experimental/Interchange/Core/Public/Types/AttributeStorage.h
--- a/Engine/Source/Runtime/Experimental/Interchange/Core/Public/Types/AttributeStorage.h
+++ b/Engine/Source/Runtime/Experimental/Interchange/Core/Public/Types/AttributeStorage.h
@@ -145,72 +145,62 @@ namespace UE
 		{
 			None				= 0,
 			Bool				= 1,
-			Box					= 2,
-			BoxSphereBounds		= 3,
-			ByteArray			= 4,
-			ByteArray64			= 5,
-			Color				= 6,
-			DateTime			= 7,
-			Double				= 8,
-			Enum				= 9,
-			Float				= 10,
-			Guid				= 11,
-			Int8				= 12,
-			Int16				= 13,
-			Int32				= 14,
-			Int64				= 15,
-			IntRect				= 16,
-			LinearColor			= 17,
-			Matrix				= 18,
-			Name				= 19,
-			Plane				= 20,
-			Quat				= 21,
-			RandomStream		= 22,
-			Rotator				= 23,
-			String				= 24,
-			Timespan			= 25,
-			Transform			= 26,
-			TwoVectors			= 27,
-			UInt8				= 28,
-			UInt16				= 29,
-			UInt32				= 30,
-			UInt64				= 31,
-			Vector				= 32,
-			Vector2d			= 33,
-			Vector4				= 34,
-			IntPoint			= 35,
-			IntVector			= 36,
-			Vector2DHalf		= 37,
-			Float16				= 38,
-			OrientedBox			= 39,
-			Sphere				= 40,
-			FrameNumber			= 41,
-			FrameRate			= 42,
-			FrameTime			= 43,
-			SoftObjectPath		= 44,
-			Matrix44f			= 45,
-			Matrix44d			= 46,
-			Plane4f				= 47,
-			Plane4d				= 48,
-			Quat4f				= 49,
-			Quat4d				= 50,
-			Rotator3f			= 51,
-			Rotator3d			= 52,
-			Transform3f			= 53,
-			Transform3d			= 54,
-			Vector3f			= 55,
-			Vector3d			= 56,
-			Vector2f			= 57,
-			Vector4f			= 58,
-			Vector4d			= 59,
-			Box2f				= 60,
-			Box2D				= 61,
-			Box3f				= 62,
-			Box3d				= 63,
-			BoxSphereBounds3f	= 64,
-			BoxSphereBounds3d	= 65,
-			Sphere3f			= 66,
-			Sphere3d			= 67
+			ByteArray			= 2,
+			ByteArray64			= 3,
+			Color				= 4,
+			DateTime			= 5,
+			Double				= 6,
+			Enum				= 7,
+			Float				= 8,
+			Guid				= 9,
+			Int8				= 10,
+			Int16				= 11,
+			Int32				= 12,
+			Int64				= 13,
+			IntRect				= 14,
+			LinearColor			= 15,
+			Name				= 16,
+			RandomStream		= 17,
+			String				= 18,
+			Timespan			= 19,
+			TwoVectors			= 20,
+			UInt8				= 21,
+			UInt16				= 22,
+			UInt32				= 23,
+			UInt64				= 24,
+			Vector2d			= 25,
+			IntPoint			= 26,
+			IntVector			= 27,
+			Vector2DHalf		= 28,
+			Float16				= 29,
+			OrientedBox			= 30,
+			FrameNumber			= 31,
+			FrameRate			= 32,
+			FrameTime			= 33,
+			SoftObjectPath		= 34,
+			Matrix44f			= 35,
+			Matrix44d			= 36,
+			Plane4f				= 37,
+			Plane4d				= 38,
+			Quat4f				= 39,
+			Quat4d				= 40,
+			Rotator3f			= 41,
+			Rotator3d			= 42,
+			Transform3f			= 43,
+			Transform3d			= 44,
+			Vector3f			= 45,
+			Vector3d			= 46,
+			Vector2f			= 47,
+			Vector4f			= 48,
+			Vector4d			= 49,
+			Box2f				= 50,
+			Box2D				= 51,
+			Box3f				= 52,
+			Box3d				= 53,
+			BoxSphereBounds3f	= 54,
+			BoxSphereBounds3d	= 55,
+			Sphere3f			= 56,
+			Sphere3d			= 57
 		};
 
 		/**
