Interchange import dialog
- Added sections to navigate between the settings of different asset types.
- Fixed the tooltips for the reimport strategy being on the wrong entries.

PropertyEditor
- Fixed an issue where a category with a space in its name would fail the section category comparison as spaces are removed from section category names.

diff --git a/Engine/Plugins/Experimental/Interchange/Editor/Source/Pipelines/Private/InterchangeEditorPipelinesModule.cpp b/Engine/Plugins/Experimental/Interchange/Editor/Source/Pipelines/Private/InterchangeEditorPipelinesModule.cpp
--- a/Engine/Plugins/Experimental/Interchange/Editor/Source/Pipelines/Private/InterchangeEditorPipelinesModule.cpp
+++ b/Engine/Plugins/Experimental/Interchange/Editor/Source/Pipelines/Private/InterchangeEditorPipelinesModule.cpp
@@ -13,6 +13,7 @@
 #include "Nodes/InterchangeBaseNode.h"
 #include "PropertyEditorModule.h"
 
+#define LOCTEXT_NAMESPACE "InterchangeEditorPipelines"
 
 class FInterchangeEditorPipelinesModule : public IInterchangeEditorPipelinesModule
 {
@@ -20,18 +21,26 @@ class FInterchangeEditorPipelinesModule : public IInterchangeEditorPipelinesModu
 	virtual void StartupModule() override;
 	virtual void ShutdownModule() override;
 
+	TSharedRef<FPropertySection> RegisterPropertySection(FPropertyEditorModule& PropertyModule, FName ClassName, FName SectionName, FText DisplayName);
+	void RegisterPropertySectionMappings();
+	void UnregisterPropertySectionMappings();
+
 private:
 	/** Pointer to the style set to use for the UI. */
 	TSharedPtr<ISlateStyle> InterchangeEditorPipelineStyle = nullptr;
+
+	TMultiMap<FName, FName> RegisteredPropertySections;
 };
 
 IMPLEMENT_MODULE(FInterchangeEditorPipelinesModule, InterchangeEditorPipelines)
 
 void FInterchangeEditorPipelinesModule::StartupModule()
 {
-	auto RegisterItems = []()
+	auto RegisterItems = [this]()
 	{
 		UInterchangeManager& InterchangeManager = UInterchangeManager::GetInterchangeManager();
+
+		RegisterPropertySectionMappings();
 	};
 
 	if (GEngine)
@@ -56,7 +65,6 @@ void FInterchangeEditorPipelinesModule::StartupModule()
 	}
 }
 
-
 void FInterchangeEditorPipelinesModule::ShutdownModule()
 {
 	FPropertyEditorModule* PropertyEditorModule = FModuleManager::GetModulePtr<FPropertyEditorModule>("PropertyEditor");
@@ -68,8 +76,73 @@ void FInterchangeEditorPipelinesModule::ShutdownModule()
 		}
 	}
 
+	UnregisterPropertySectionMappings();
+
 	InterchangeEditorPipelineStyle = nullptr;
 }
 
+TSharedRef<FPropertySection> FInterchangeEditorPipelinesModule::RegisterPropertySection(FPropertyEditorModule& PropertyModule, FName ClassName, FName SectionName, FText DisplayName)
+{
+	TSharedRef<FPropertySection> PropertySection = PropertyModule.FindOrCreateSection(ClassName, SectionName, DisplayName);
+	RegisteredPropertySections.Add(ClassName, SectionName);
+
+	return PropertySection;
+}
+
+void FInterchangeEditorPipelinesModule::RegisterPropertySectionMappings()
+{
+	const FName PropertyEditorModuleName("PropertyEditor");
+	FPropertyEditorModule& PropertyModule = FModuleManager::GetModuleChecked<FPropertyEditorModule>(PropertyEditorModuleName);
+
+	// Assets
+	{
+		TSharedRef<FPropertySection> Section = RegisterPropertySection(PropertyModule, "InterchangeGenericAssetsPipeline", "General", LOCTEXT("General", "General"));
+		Section->AddCategory("Common");
+	}
+
+	// Materials
+	{
+		TSharedRef<FPropertySection> Section = RegisterPropertySection(PropertyModule, "InterchangeGenericMaterialPipeline", "Materials", LOCTEXT("Materials", "Materials"));
+		Section->AddCategory("Materials");
+	}
+
+	// Skeletal Meshes
+	{
+		TSharedRef<FPropertySection> Section = RegisterPropertySection(PropertyModule, "InterchangeGenericMeshPipeline", "SkeletalMeshes", LOCTEXT("Skeletal Meshes", "Skeletal Meshes"));
+		Section->AddCategory("Common Meshes");
+		Section->AddCategory("Skeletal Meshes");
+	}
+
+	// Static Meshes
+	{
+		TSharedRef<FPropertySection> Section = RegisterPropertySection(PropertyModule, "InterchangeGenericMeshPipeline", "StaticMeshes", LOCTEXT("Static Meshes", "Static Meshes"));
+		Section->AddCategory("Common Meshes");
+		Section->AddCategory("Static Meshes");
+	}
+
+	// Textures
+	{
+		TSharedRef<FPropertySection> Section = RegisterPropertySection(PropertyModule, "InterchangeGenericTexturePipeline", "Textures", LOCTEXT("Textures", "Textures"));
+		Section->AddCategory("Textures");
+	}
+}
+
+void FInterchangeEditorPipelinesModule::UnregisterPropertySectionMappings()
+{
+	const FName PropertyEditorModuleName("PropertyEditor");
+	FPropertyEditorModule* PropertyModule = FModuleManager::GetModulePtr<FPropertyEditorModule>(PropertyEditorModuleName);
+
+	if (!PropertyModule)
+	{
+		return;
+	}
+
+	for (TMultiMap<FName, FName>::TIterator PropertySectionIterator = RegisteredPropertySections.CreateIterator(); PropertySectionIterator; ++PropertySectionIterator)
+	{
+		PropertyModule->RemoveSection(PropertySectionIterator->Key, PropertySectionIterator->Value);
+		PropertySectionIterator.RemoveCurrent();
+	}
+}
 
+#undef LOCTEXT_NAMESPACE
 
diff --git a/Engine/Plugins/Experimental/Interchange/Editor/Source/Pipelines/Private/SInterchangePipelineConfigurationDialog.cpp b/Engine/Plugins/Experimental/Interchange/Editor/Source/Pipelines/Private/SInterchangePipelineConfigurationDialog.cpp
--- a/Engine/Plugins/Experimental/Interchange/Editor/Source/Pipelines/Private/SInterchangePipelineConfigurationDialog.cpp
+++ b/Engine/Plugins/Experimental/Interchange/Editor/Source/Pipelines/Private/SInterchangePipelineConfigurationDialog.cpp
@@ -390,7 +390,9 @@ TSharedRef<SBox> SInterchangePipelineConfigurationDialog::SpawnPipelineConfigura
 
 	FPropertyEditorModule& PropertyEditorModule = FModuleManager::GetModuleChecked<FPropertyEditorModule>("PropertyEditor");
 	FDetailsViewArgs DetailsViewArgs;
-	DetailsViewArgs.bAllowSearch = false;
+	DetailsViewArgs.bAllowSearch = true;
+	DetailsViewArgs.bShowPropertyMatrixButton = false;
+	DetailsViewArgs.bShowSectionSelector = true;
 	DetailsViewArgs.NameAreaSettings = FDetailsViewArgs::HideNameArea;
 	PipelineConfigurationDetailsView = PropertyEditorModule.CreateDetailView(DetailsViewArgs);
 	InspectorBox->SetContent(PipelineConfigurationDetailsView->AsShared());
diff --git a/Engine/Source/Editor/PropertyEditor/Private/PropertyEditorModule.cpp b/Engine/Source/Editor/PropertyEditor/Private/PropertyEditorModule.cpp
--- a/Engine/Source/Editor/PropertyEditor/Private/PropertyEditorModule.cpp
+++ b/Engine/Source/Editor/PropertyEditor/Private/PropertyEditorModule.cpp
@@ -518,12 +518,18 @@ void FPropertySection::RemoveCategory(FName CategoryName)
 
 bool FPropertySection::HasAddedCategory(FName CategoryName) const
 {
-	return AddedCategories.Contains(CategoryName);
+	FString CategoryString = CategoryName.ToString();
+	CategoryString.RemoveSpacesInline();
+
+	return AddedCategories.Contains(*CategoryString);
 }
 
 bool FPropertySection::HasRemovedCategory(FName CategoryName) const
 {
-	return RemovedCategories.Contains(CategoryName);
+	FString CategoryString = CategoryName.ToString();
+	CategoryString.RemoveSpacesInline();
+
+	return RemovedCategories.Contains(*CategoryString);
 }
 
 FClassSectionMapping::FClassSectionMapping(FName InClassName) :
diff --git a/Engine/Source/Runtime/Experimental/Interchange/Core/Public/Nodes/InterchangeBaseNode.h b/Engine/Source/Runtime/Experimental/Interchange/Core/Public/Nodes/InterchangeBaseNode.h
--- a/Engine/Source/Runtime/Experimental/Interchange/Core/Public/Nodes/InterchangeBaseNode.h
+++ b/Engine/Source/Runtime/Experimental/Interchange/Core/Public/Nodes/InterchangeBaseNode.h
@@ -219,9 +219,12 @@ enum class EInterchangeNodeContainerType : uint8
 UENUM()
 enum class EReimportStrategyFlags : uint8
 {
-	ApplyNoProperties, //Do not apply any property when re-importing, simply change the source data
-	ApplyPipelineProperties, //Always apply all pipeline specified properties
-	ApplyEditorChangedProperties //Always apply all pipeline properties, but leave the properties modified in editor since the last import
+	//Do not apply any property when re-importing, simply change the source data
+	ApplyNoProperties,
+	//Always apply all pipeline specified properties
+	ApplyPipelineProperties,
+	//Always apply all pipeline properties, but leave the properties modified in editor since the last import
+	ApplyEditorChangedProperties
 };
 
 /**
