Add a translator API function to get a list of extensions for a specific asset. Also add a file picker Interchange can use from its runtime code.
#rnx

diff --git a/Engine/Config/BaseEngine.ini b/Engine/Config/BaseEngine.ini
--- a/Engine/Config/BaseEngine.ini
+++ b/Engine/Config/BaseEngine.ini
@@ -2979,6 +2979,7 @@ DefaultPipelineStack=Generic
 DefaultScenePipelineStack=Scene
 PipelineConfigurationDialogClass=/Script/InterchangeEditorPipelines.InterchangePipelineConfigurationGeneric
 bShowPipelineStacksConfigurationDialog=True
+FilePickerClass=/Script/InterchangeEditorUtilities.InterchangeFilePickerGeneric
 
 [AlternateTextureCompression]
 TextureCompressionFormat="TextureFormatOodle"
diff --git a/Engine/Plugins/Experimental/Interchange/Editor/InterchangeEditor.uplugin b/Engine/Plugins/Experimental/Interchange/Editor/InterchangeEditor.uplugin
--- a/Engine/Plugins/Experimental/Interchange/Editor/InterchangeEditor.uplugin
+++ b/Engine/Plugins/Experimental/Interchange/Editor/InterchangeEditor.uplugin
@@ -29,6 +29,11 @@
 			"Name": "InterchangeResultsBrowser",
 			"Type": "Editor",
 			"LoadingPhase": "Default"
+		},
+		{
+			"Name": "InterchangeEditorUtilities",
+			"Type": "Editor",
+			"LoadingPhase": "Default"
 		}
 	],
 	"Plugins": [
diff --git a/Engine/Plugins/Experimental/Interchange/Editor/Source/Utilities/InterchangeEditorUtilities.Build.cs b/Engine/Plugins/Experimental/Interchange/Editor/Source/Utilities/InterchangeEditorUtilities.Build.cs
new file mode 100644
--- /dev/null
+++ b/Engine/Plugins/Experimental/Interchange/Editor/Source/Utilities/InterchangeEditorUtilities.Build.cs
@@ -0,0 +1,30 @@
+// Copyright Epic Games, Inc. All Rights Reserved.
+
+namespace UnrealBuildTool.Rules
+{
+	public class InterchangeEditorUtilities : ModuleRules
+	{
+		public InterchangeEditorUtilities(ReadOnlyTargetRules Target) : base(Target)
+		{
+			PublicDependencyModuleNames.AddRange(
+				new string[]
+				{
+					"Core",
+					"CoreUObject",
+					"Engine",
+					"InterchangeCore",
+					"InterchangeEngine"
+				}
+			);
+
+			PrivateDependencyModuleNames.AddRange(
+				new string[]
+				{
+					"DesktopPlatform",
+					"Slate",
+					"SlateCore",
+				}
+			);
+		}
+	}
+}
diff --git a/Engine/Plugins/Experimental/Interchange/Editor/Source/Utilities/Private/InterchangeEditorUtilitiesModule.cpp b/Engine/Plugins/Experimental/Interchange/Editor/Source/Utilities/Private/InterchangeEditorUtilitiesModule.cpp
new file mode 100644
--- /dev/null
+++ b/Engine/Plugins/Experimental/Interchange/Editor/Source/Utilities/Private/InterchangeEditorUtilitiesModule.cpp
@@ -0,0 +1,24 @@
+// Copyright Epic Games, Inc. All Rights Reserved.
+
+#include "InterchangeEditorUtilitiesModule.h"
+
+
+#define LOCTEXT_NAMESPACE "InterchangeEditorUtilities"
+
+class FInterchangeEditorUtilitiesModule : public IInterchangeEditorUtilitiesModule
+{
+	/** IModuleInterface implementation */
+	virtual void StartupModule() override;
+	virtual void ShutdownModule() override;
+
+};
+
+IMPLEMENT_MODULE(FInterchangeEditorUtilitiesModule, InterchangeEditorUtilities)
+
+void FInterchangeEditorUtilitiesModule::StartupModule()
+{
+}
+
+void FInterchangeEditorUtilitiesModule::ShutdownModule()
+{
+}
diff --git a/Engine/Plugins/Experimental/Interchange/Editor/Source/Utilities/Private/InterchangeOpenFileDialog.cpp b/Engine/Plugins/Experimental/Interchange/Editor/Source/Utilities/Private/InterchangeOpenFileDialog.cpp
new file mode 100644
--- /dev/null
+++ b/Engine/Plugins/Experimental/Interchange/Editor/Source/Utilities/Private/InterchangeOpenFileDialog.cpp
@@ -0,0 +1,111 @@
+// Copyright Epic Games, Inc. All Rights Reserved. 
+#include "InterchangeOpenFileDialog.h"
+
+#include "Containers/UnrealString.h"
+#include "DesktopPlatformModule.h"
+#include "IDesktopPlatform.h"
+#include "InterchangeFilePickerBase.h"
+#include "InterchangeManager.h"
+#include "InterchangeTranslatorBase.h"
+#include "UObject/Object.h"
+#include "UObject/ObjectMacros.h"
+#include "Framework/Application/SlateApplication.h"
+
+namespace UE::Interchange::Utilities::Private
+{
+	FString GetOpenFileDialogExtensions(const TArray<FString>& TranslatorFormats)
+	{
+		//Each format should be layout like this "fbx;Filmbox"
+		if (TranslatorFormats.Num() == 0)
+		{
+			return FString();
+		}
+		TArray<FString> Extensions;
+		TArray<FString> ExtensionDescriptions;
+		Extensions.Reserve(TranslatorFormats.Num());
+		ExtensionDescriptions.Reserve(TranslatorFormats.Num());
+		for (const FString& Format : TranslatorFormats)
+		{
+			int32 DelimiterIndex;
+			if (Format.FindChar(';', DelimiterIndex))
+			{
+				Extensions.Add(Format.Left(DelimiterIndex));
+				ExtensionDescriptions.Add(Format.RightChop(DelimiterIndex + 1));
+			}
+		}
+		const int32 ExtensionCount = Extensions.Num();
+		if (ExtensionCount == 0)
+		{
+			return FString();
+		}
+
+		FString ExtensionStrAll = TEXT("All file extensions|");
+		FString ExtensionStr;
+		for (int32 ExtensionIndex = 0; ExtensionIndex < ExtensionCount; ++ExtensionIndex)
+		{
+			if (ExtensionIndex > 0)
+			{
+				ExtensionStrAll += TEXT(";*.");
+			}
+			else
+			{
+				ExtensionStrAll += TEXT("*.");
+			}
+			ExtensionStrAll += Extensions[ExtensionIndex];
+			ExtensionStr += ExtensionDescriptions[ExtensionIndex] + TEXT("|*.") + Extensions[ExtensionIndex];
+			if ((ExtensionIndex + 1) < ExtensionCount)
+			{
+				ExtensionStr += TEXT("|");
+			}
+		}
+		ExtensionStr += TEXT("Any files|*.*");
+		ExtensionStrAll += TEXT("|") + ExtensionStr;
+
+		return ExtensionStrAll;
+	}
+
+	bool FilePickerDialog(const FString& Extensions, FString& OutFilename)
+	{
+		// First, display the file open dialog for selecting the file.
+		TArray<FString> OpenFilenames;
+		IDesktopPlatform* DesktopPlatform = FDesktopPlatformModule::Get();
+		bool bOpen = false;
+		if (DesktopPlatform)
+		{
+			FText PromptTitle = NSLOCTEXT("InterchangeUtilities_OpenFileDialog", "FilePickerDialog", "Select a file");
+			bOpen = DesktopPlatform->OpenFileDialog(
+				FSlateApplication::Get().FindBestParentWindowHandleForDialogs(nullptr),
+				PromptTitle.ToString(),
+				TEXT(""),
+				TEXT(""),
+				*Extensions,
+				EFileDialogFlags::None,
+				OpenFilenames
+			);
+		}
+		FString SelectFilename;
+		// Only continue if we pressed OK and have only one file selected.
+		if (bOpen)
+		{
+			if (OpenFilenames.Num() >= 1)
+			{
+				OutFilename = OpenFilenames[0];
+				return true;
+			}
+		}
+
+		return false;
+	}
+} //ns UE::Interchange::Utilities::Private
+
+bool UInterchangeFilePickerGeneric::FilePickerForTranslatorAssetType(const EInterchangeTranslatorAssetType TranslatorAssetType, FString& OutFilename)
+{
+	FString Extensions = UE::Interchange::Utilities::Private::GetOpenFileDialogExtensions(UInterchangeManager::GetInterchangeManager().GetSupportedAssetTypeFormats(TranslatorAssetType));
+	return UE::Interchange::Utilities::Private::FilePickerDialog(Extensions, OutFilename);
+}
+
+bool UInterchangeFilePickerGeneric::FilePickerForTranslatorType(const EInterchangeTranslatorType TranslatorType, FString& OutFilename)
+{
+	FString Extensions = UE::Interchange::Utilities::Private::GetOpenFileDialogExtensions(UInterchangeManager::GetInterchangeManager().GetSupportedFormats(TranslatorType));
+	return UE::Interchange::Utilities::Private::FilePickerDialog(Extensions, OutFilename);
+}
diff --git a/Engine/Plugins/Experimental/Interchange/Editor/Source/Utilities/Public/InterchangeEditorUtilitiesModule.h b/Engine/Plugins/Experimental/Interchange/Editor/Source/Utilities/Public/InterchangeEditorUtilitiesModule.h
new file mode 100644
--- /dev/null
+++ b/Engine/Plugins/Experimental/Interchange/Editor/Source/Utilities/Public/InterchangeEditorUtilitiesModule.h
@@ -0,0 +1,40 @@
+// Copyright Epic Games, Inc. All Rights Reserved.
+
+#pragma once
+
+#include "CoreMinimal.h"
+#include "Modules/ModuleInterface.h"
+#include "Modules/ModuleManager.h"
+#include "Styling/ISlateStyle.h"
+
+
+/**
+ * The public interface to this module
+ */
+class IInterchangeEditorUtilitiesModule : public IModuleInterface
+{
+
+public:
+
+	/**
+	 * Singleton-like access to this module's interface.  This is just for convenience!
+	 * Beware of calling this during the shutdown phase, though.  Your module might have been unloaded already.
+	 *
+	 * @return Returns singleton instance, loading the module on demand if needed
+	 */
+	static inline IInterchangeEditorUtilitiesModule& Get()
+	{
+		return FModuleManager::LoadModuleChecked< IInterchangeEditorUtilitiesModule >( "InterchangeEditorUtilities" );
+	}
+
+	/**
+	 * Checks to see if this module is loaded and ready.  It is only valid to call Get() if IsAvailable() returns true.
+	 *
+	 * @return True if the module is loaded and ready to use
+	 */
+	static inline bool IsAvailable()
+	{
+		return FModuleManager::Get().IsModuleLoaded( "InterchangeEditorUtilities" );
+	}
+};
+
diff --git a/Engine/Plugins/Experimental/Interchange/Editor/Source/Utilities/Public/InterchangeOpenFileDialog.h b/Engine/Plugins/Experimental/Interchange/Editor/Source/Utilities/Public/InterchangeOpenFileDialog.h
new file mode 100644
--- /dev/null
+++ b/Engine/Plugins/Experimental/Interchange/Editor/Source/Utilities/Public/InterchangeOpenFileDialog.h
@@ -0,0 +1,24 @@
+// Copyright Epic Games, Inc. All Rights Reserved.
+
+#pragma once
+
+#include "Containers/UnrealString.h"
+#include "InterchangeFilePickerBase.h"
+#include "InterchangeTranslatorBase.h"
+#include "UObject/Object.h"
+#include "UObject/ObjectMacros.h"
+
+#include "InterchangeOpenFileDialog.generated.h"
+
+UCLASS(BlueprintType, Blueprintable)
+class INTERCHANGEEDITORUTILITIES_API UInterchangeFilePickerGeneric : public UInterchangeFilePickerBase
+{
+	GENERATED_BODY()
+
+public:
+
+protected:
+
+	virtual bool FilePickerForTranslatorAssetType(const EInterchangeTranslatorAssetType TranslatorAssetType, FString& OutFilename) override;
+	virtual bool FilePickerForTranslatorType(const EInterchangeTranslatorType TranslatorType, FString& OutFilename) override;
+};
diff --git a/Engine/Plugins/Experimental/Interchange/Runtime/Source/Import/Private/Fbx/InterchangeFbxTranslator.cpp b/Engine/Plugins/Experimental/Interchange/Runtime/Source/Import/Private/Fbx/InterchangeFbxTranslator.cpp
--- a/Engine/Plugins/Experimental/Interchange/Runtime/Source/Import/Private/Fbx/InterchangeFbxTranslator.cpp
+++ b/Engine/Plugins/Experimental/Interchange/Runtime/Source/Import/Private/Fbx/InterchangeFbxTranslator.cpp
@@ -34,6 +34,17 @@ EInterchangeTranslatorType UInterchangeFbxTranslator::GetTranslatorType() const
 	return EInterchangeTranslatorType::Scenes;
 }
 
+bool UInterchangeFbxTranslator::DoesSupportAssetType(EInterchangeTranslatorAssetType AssetType) const
+{
+	//fbx translator support Meshes, Materials and animation
+	if (AssetType == EInterchangeTranslatorAssetType::Materials || AssetType == EInterchangeTranslatorAssetType::Meshes || AssetType == EInterchangeTranslatorAssetType::Animations)
+	{
+		return true;
+	}
+
+	return false;
+}
+
 TArray<FString> UInterchangeFbxTranslator::GetSupportedFormats() const
 {
 	TArray<FString> Formats {TEXT("fbx;Filmbox")};
diff --git a/Engine/Plugins/Experimental/Interchange/Runtime/Source/Import/Private/Gltf/InterchangeGltfTranslator.cpp b/Engine/Plugins/Experimental/Interchange/Runtime/Source/Import/Private/Gltf/InterchangeGltfTranslator.cpp
--- a/Engine/Plugins/Experimental/Interchange/Runtime/Source/Import/Private/Gltf/InterchangeGltfTranslator.cpp
+++ b/Engine/Plugins/Experimental/Interchange/Runtime/Source/Import/Private/Gltf/InterchangeGltfTranslator.cpp
@@ -692,6 +692,17 @@ EInterchangeTranslatorType UInterchangeGltfTranslator::GetTranslatorType() const
 	return EInterchangeTranslatorType::Scenes;
 }
 
+bool UInterchangeGltfTranslator::DoesSupportAssetType(EInterchangeTranslatorAssetType AssetType) const
+{
+	//gltf translator support Meshes and Materials
+	if (AssetType == EInterchangeTranslatorAssetType::Materials || AssetType == EInterchangeTranslatorAssetType::Meshes)
+	{
+		return true;
+	}
+
+	return false;
+}
+
 TArray<FString> UInterchangeGltfTranslator::GetSupportedFormats() const
 {
 	TArray<FString> GltfExtensions;
diff --git a/Engine/Plugins/Experimental/Interchange/Runtime/Source/Import/Private/Mesh/InterchangeOBJTranslator.cpp b/Engine/Plugins/Experimental/Interchange/Runtime/Source/Import/Private/Mesh/InterchangeOBJTranslator.cpp
--- a/Engine/Plugins/Experimental/Interchange/Runtime/Source/Import/Private/Mesh/InterchangeOBJTranslator.cpp
+++ b/Engine/Plugins/Experimental/Interchange/Runtime/Source/Import/Private/Mesh/InterchangeOBJTranslator.cpp
@@ -822,6 +822,18 @@ TArray<FString> UInterchangeOBJTranslator::GetSupportedFormats() const
 	return Formats;
 }
 
+bool UInterchangeOBJTranslator::DoesSupportAssetType(EInterchangeTranslatorAssetType AssetType) const
+{
+	//Obj translator support Meshes and Materials
+	if (AssetType == EInterchangeTranslatorAssetType::Materials || AssetType == EInterchangeTranslatorAssetType::Meshes)
+	{
+		return true;
+	}
+
+	return false;
+}
+
+
 bool UInterchangeOBJTranslator::Translate(UInterchangeBaseNodeContainer& BaseNodeContainer) const
 {
 	// Get filename from SourceData
diff --git a/Engine/Plugins/Experimental/Interchange/Runtime/Source/Import/Public/Fbx/InterchangeFbxTranslator.h b/Engine/Plugins/Experimental/Interchange/Runtime/Source/Import/Public/Fbx/InterchangeFbxTranslator.h
--- a/Engine/Plugins/Experimental/Interchange/Runtime/Source/Import/Public/Fbx/InterchangeFbxTranslator.h
+++ b/Engine/Plugins/Experimental/Interchange/Runtime/Source/Import/Public/Fbx/InterchangeFbxTranslator.h
@@ -30,7 +30,8 @@ public:
 	UInterchangeFbxTranslator();
 
 	/** Begin UInterchangeTranslatorBase API*/
-	EInterchangeTranslatorType GetTranslatorType() const override;
+	virtual EInterchangeTranslatorType GetTranslatorType() const override;
+	virtual bool DoesSupportAssetType(EInterchangeTranslatorAssetType AssetType) const override;
 	virtual TArray<FString> GetSupportedFormats() const override;
 	virtual bool Translate(UInterchangeBaseNodeContainer& BaseNodeContainer) const override;
 	virtual void ReleaseSource() override;
diff --git a/Engine/Plugins/Experimental/Interchange/Runtime/Source/Import/Public/Gltf/InterchangeGltfTranslator.h b/Engine/Plugins/Experimental/Interchange/Runtime/Source/Import/Public/Gltf/InterchangeGltfTranslator.h
--- a/Engine/Plugins/Experimental/Interchange/Runtime/Source/Import/Public/Gltf/InterchangeGltfTranslator.h
+++ b/Engine/Plugins/Experimental/Interchange/Runtime/Source/Import/Public/Gltf/InterchangeGltfTranslator.h
@@ -26,6 +26,7 @@ class UInterchangeGltfTranslator : public UInterchangeTranslatorBase,
 public:
 	/** Begin UInterchangeTranslatorBase API*/
 	virtual EInterchangeTranslatorType GetTranslatorType() const override;
+	virtual bool DoesSupportAssetType(EInterchangeTranslatorAssetType AssetType) const override;
 	virtual TArray<FString> GetSupportedFormats() const override;
 	virtual bool Translate( UInterchangeBaseNodeContainer& BaseNodeContainer ) const override;
 	/** End UInterchangeTranslatorBase API*/
diff --git a/Engine/Plugins/Experimental/Interchange/Runtime/Source/Import/Public/Mesh/InterchangeOBJTranslator.h b/Engine/Plugins/Experimental/Interchange/Runtime/Source/Import/Public/Mesh/InterchangeOBJTranslator.h
--- a/Engine/Plugins/Experimental/Interchange/Runtime/Source/Import/Public/Mesh/InterchangeOBJTranslator.h
+++ b/Engine/Plugins/Experimental/Interchange/Runtime/Source/Import/Public/Mesh/InterchangeOBJTranslator.h
@@ -31,7 +31,7 @@ public:
 
 
 	virtual TArray<FString> GetSupportedFormats() const override;
-
+	virtual bool DoesSupportAssetType(EInterchangeTranslatorAssetType AssetType) const override;
 	/**
 	 * Translate the associated source data into a node hold by the specified nodes container.
 	 *
diff --git a/Engine/Plugins/Experimental/Interchange/Runtime/Source/Import/Public/Texture/InterchangeBMPTranslator.h b/Engine/Plugins/Experimental/Interchange/Runtime/Source/Import/Public/Texture/InterchangeBMPTranslator.h
--- a/Engine/Plugins/Experimental/Interchange/Runtime/Source/Import/Public/Texture/InterchangeBMPTranslator.h
+++ b/Engine/Plugins/Experimental/Interchange/Runtime/Source/Import/Public/Texture/InterchangeBMPTranslator.h
@@ -18,8 +18,8 @@ class INTERCHANGEIMPORT_API UInterchangeBMPTranslator : public UInterchangeTrans
 	GENERATED_BODY()
 public:
 
-	TArray<FString> GetSupportedFormats() const override;
-
+	virtual TArray<FString> GetSupportedFormats() const override;
+	virtual bool DoesSupportAssetType(EInterchangeTranslatorAssetType AssetType) const override { return AssetType == EInterchangeTranslatorAssetType::Textures; }
 	/**
 	 * Translate the associated source data into a node hold by the specified nodes container.
 	 *
diff --git a/Engine/Plugins/Experimental/Interchange/Runtime/Source/Import/Public/Texture/InterchangeDDSTranslator.h b/Engine/Plugins/Experimental/Interchange/Runtime/Source/Import/Public/Texture/InterchangeDDSTranslator.h
--- a/Engine/Plugins/Experimental/Interchange/Runtime/Source/Import/Public/Texture/InterchangeDDSTranslator.h
+++ b/Engine/Plugins/Experimental/Interchange/Runtime/Source/Import/Public/Texture/InterchangeDDSTranslator.h
@@ -20,7 +20,8 @@ class INTERCHANGEIMPORT_API UInterchangeDDSTranslator : public UInterchangeTrans
 	GENERATED_BODY()
 public:
 
-	TArray<FString> GetSupportedFormats() const override;
+	virtual TArray<FString> GetSupportedFormats() const override;
+	virtual bool DoesSupportAssetType(EInterchangeTranslatorAssetType AssetType) const override { return AssetType == EInterchangeTranslatorAssetType::Textures; }
 	
 	/*
 	 * return true if the translator can translate the specified source data.
diff --git a/Engine/Plugins/Experimental/Interchange/Runtime/Source/Import/Public/Texture/InterchangeEXRTranslator.h b/Engine/Plugins/Experimental/Interchange/Runtime/Source/Import/Public/Texture/InterchangeEXRTranslator.h
--- a/Engine/Plugins/Experimental/Interchange/Runtime/Source/Import/Public/Texture/InterchangeEXRTranslator.h
+++ b/Engine/Plugins/Experimental/Interchange/Runtime/Source/Import/Public/Texture/InterchangeEXRTranslator.h
@@ -18,7 +18,8 @@ class INTERCHANGEIMPORT_API UInterchangeEXRTranslator : public UInterchangeTrans
 	GENERATED_BODY()
 public:
 	
-	TArray<FString> GetSupportedFormats() const override;
+	virtual TArray<FString> GetSupportedFormats() const override;
+	virtual bool DoesSupportAssetType(EInterchangeTranslatorAssetType AssetType) const override { return AssetType == EInterchangeTranslatorAssetType::Textures; }
 
 	/**
 	 * Translate the associated source data into a node hold by the specified nodes container.
diff --git a/Engine/Plugins/Experimental/Interchange/Runtime/Source/Import/Public/Texture/InterchangeHDRTranslator.h b/Engine/Plugins/Experimental/Interchange/Runtime/Source/Import/Public/Texture/InterchangeHDRTranslator.h
--- a/Engine/Plugins/Experimental/Interchange/Runtime/Source/Import/Public/Texture/InterchangeHDRTranslator.h
+++ b/Engine/Plugins/Experimental/Interchange/Runtime/Source/Import/Public/Texture/InterchangeHDRTranslator.h
@@ -18,7 +18,8 @@ class INTERCHANGEIMPORT_API UInterchangeHDRTranslator : public UInterchangeTrans
 	GENERATED_BODY()
 public:
 
-	TArray<FString> GetSupportedFormats() const override;
+	virtual TArray<FString> GetSupportedFormats() const override;
+	virtual bool DoesSupportAssetType(EInterchangeTranslatorAssetType AssetType) const override { return AssetType == EInterchangeTranslatorAssetType::Textures; }
 
 	/**
 	 * Translate the associated source data into a node hold by the specified nodes container.
diff --git a/Engine/Plugins/Experimental/Interchange/Runtime/Source/Import/Public/Texture/InterchangeIESTranslator.h b/Engine/Plugins/Experimental/Interchange/Runtime/Source/Import/Public/Texture/InterchangeIESTranslator.h
--- a/Engine/Plugins/Experimental/Interchange/Runtime/Source/Import/Public/Texture/InterchangeIESTranslator.h
+++ b/Engine/Plugins/Experimental/Interchange/Runtime/Source/Import/Public/Texture/InterchangeIESTranslator.h
@@ -18,7 +18,8 @@ class INTERCHANGEIMPORT_API UInterchangeIESTranslator : public UInterchangeTrans
 	GENERATED_BODY()
 public:
 
-	TArray<FString> GetSupportedFormats() const override;
+	virtual TArray<FString> GetSupportedFormats() const override;
+	virtual bool DoesSupportAssetType(EInterchangeTranslatorAssetType AssetType) const override { return AssetType == EInterchangeTranslatorAssetType::Textures; }
 
 	/**
 	 * Translate the associated source data into a node hold by the specified nodes container.
diff --git a/Engine/Plugins/Experimental/Interchange/Runtime/Source/Import/Public/Texture/InterchangeJPGTranslator.h b/Engine/Plugins/Experimental/Interchange/Runtime/Source/Import/Public/Texture/InterchangeJPGTranslator.h
--- a/Engine/Plugins/Experimental/Interchange/Runtime/Source/Import/Public/Texture/InterchangeJPGTranslator.h
+++ b/Engine/Plugins/Experimental/Interchange/Runtime/Source/Import/Public/Texture/InterchangeJPGTranslator.h
@@ -18,7 +18,8 @@ class INTERCHANGEIMPORT_API UInterchangeJPGTranslator : public UInterchangeTrans
 	GENERATED_BODY()
 public:
 	
-	TArray<FString> GetSupportedFormats() const override;
+	virtual TArray<FString> GetSupportedFormats() const override;
+	virtual bool DoesSupportAssetType(EInterchangeTranslatorAssetType AssetType) const override { return AssetType == EInterchangeTranslatorAssetType::Textures; }
 
 	/**
 	 * Translate the associated source data into a node hold by the specified nodes container.
diff --git a/Engine/Plugins/Experimental/Interchange/Runtime/Source/Import/Public/Texture/InterchangePCXTranslator.h b/Engine/Plugins/Experimental/Interchange/Runtime/Source/Import/Public/Texture/InterchangePCXTranslator.h
--- a/Engine/Plugins/Experimental/Interchange/Runtime/Source/Import/Public/Texture/InterchangePCXTranslator.h
+++ b/Engine/Plugins/Experimental/Interchange/Runtime/Source/Import/Public/Texture/InterchangePCXTranslator.h
@@ -18,7 +18,8 @@ class INTERCHANGEIMPORT_API UInterchangePCXTranslator : public UInterchangeTrans
 	GENERATED_BODY()
 public:
 	
-	TArray<FString> GetSupportedFormats() const override;
+	virtual TArray<FString> GetSupportedFormats() const override;
+	virtual bool DoesSupportAssetType(EInterchangeTranslatorAssetType AssetType) const override { return AssetType == EInterchangeTranslatorAssetType::Textures; }
 
 	/**
 	 * Translates the associated source data into a node hold by the specified nodes container.
diff --git a/Engine/Plugins/Experimental/Interchange/Runtime/Source/Import/Public/Texture/InterchangePNGTranslator.h b/Engine/Plugins/Experimental/Interchange/Runtime/Source/Import/Public/Texture/InterchangePNGTranslator.h
--- a/Engine/Plugins/Experimental/Interchange/Runtime/Source/Import/Public/Texture/InterchangePNGTranslator.h
+++ b/Engine/Plugins/Experimental/Interchange/Runtime/Source/Import/Public/Texture/InterchangePNGTranslator.h
@@ -18,7 +18,8 @@ class INTERCHANGEIMPORT_API UInterchangePNGTranslator : public UInterchangeTrans
 	GENERATED_BODY()
 public:
 
-	TArray<FString> GetSupportedFormats() const override;
+	virtual TArray<FString> GetSupportedFormats() const override;
+	virtual bool DoesSupportAssetType(EInterchangeTranslatorAssetType AssetType) const override { return AssetType == EInterchangeTranslatorAssetType::Textures; }
 
 	/**
 	 * Translate the associated source data into a node hold by the specified nodes container.
diff --git a/Engine/Plugins/Experimental/Interchange/Runtime/Source/Import/Public/Texture/InterchangePSDTranslator.h b/Engine/Plugins/Experimental/Interchange/Runtime/Source/Import/Public/Texture/InterchangePSDTranslator.h
--- a/Engine/Plugins/Experimental/Interchange/Runtime/Source/Import/Public/Texture/InterchangePSDTranslator.h
+++ b/Engine/Plugins/Experimental/Interchange/Runtime/Source/Import/Public/Texture/InterchangePSDTranslator.h
@@ -18,7 +18,8 @@ class INTERCHANGEIMPORT_API UInterchangePSDTranslator : public UInterchangeTrans
 	GENERATED_BODY()
 public:
 
-	TArray<FString> GetSupportedFormats() const override;
+	virtual TArray<FString> GetSupportedFormats() const override;
+	virtual bool DoesSupportAssetType(EInterchangeTranslatorAssetType AssetType) const override { return AssetType == EInterchangeTranslatorAssetType::Textures; }
 
 	/**
 	 * Translate the associated source data into a node hold by the specified nodes container.
diff --git a/Engine/Plugins/Experimental/Interchange/Runtime/Source/Import/Public/Texture/InterchangeTGATranslator.h b/Engine/Plugins/Experimental/Interchange/Runtime/Source/Import/Public/Texture/InterchangeTGATranslator.h
--- a/Engine/Plugins/Experimental/Interchange/Runtime/Source/Import/Public/Texture/InterchangeTGATranslator.h
+++ b/Engine/Plugins/Experimental/Interchange/Runtime/Source/Import/Public/Texture/InterchangeTGATranslator.h
@@ -17,7 +17,8 @@ class INTERCHANGEIMPORT_API UInterchangeTGATranslator : public UInterchangeTrans
 	GENERATED_BODY()
 public:
 	
-	TArray<FString> GetSupportedFormats() const override;
+	virtual TArray<FString> GetSupportedFormats() const override;
+	virtual bool DoesSupportAssetType(EInterchangeTranslatorAssetType AssetType) const override { return AssetType == EInterchangeTranslatorAssetType::Textures; }
 
 	/**
 	 * Translates the associated source data into a node hold by the specified nodes container.
diff --git a/Engine/Plugins/Experimental/Interchange/Runtime/Source/Import/Public/Texture/InterchangeTIFFTranslator.h b/Engine/Plugins/Experimental/Interchange/Runtime/Source/Import/Public/Texture/InterchangeTIFFTranslator.h
--- a/Engine/Plugins/Experimental/Interchange/Runtime/Source/Import/Public/Texture/InterchangeTIFFTranslator.h
+++ b/Engine/Plugins/Experimental/Interchange/Runtime/Source/Import/Public/Texture/InterchangeTIFFTranslator.h
@@ -18,7 +18,8 @@ class INTERCHANGEIMPORT_API UInterchangeTIFFTranslator : public UInterchangeTran
 	GENERATED_BODY()
 public:
 
-	TArray<FString> GetSupportedFormats() const override;
+	virtual TArray<FString> GetSupportedFormats() const override;
+	virtual bool DoesSupportAssetType(EInterchangeTranslatorAssetType AssetType) const override { return AssetType == EInterchangeTranslatorAssetType::Textures; }
 
 	/**
 	 * Translate a source data into a node hold by the specified nodes container.
diff --git a/Engine/Source/Runtime/Experimental/Interchange/Core/Public/InterchangeTranslatorBase.h b/Engine/Source/Runtime/Experimental/Interchange/Core/Public/InterchangeTranslatorBase.h
--- a/Engine/Source/Runtime/Experimental/Interchange/Core/Public/InterchangeTranslatorBase.h
+++ b/Engine/Source/Runtime/Experimental/Interchange/Core/Public/InterchangeTranslatorBase.h
@@ -20,6 +20,16 @@ enum class EInterchangeTranslatorType : uint8
 	Scenes = Assets | Actors
 };
 
+UENUM(BlueprintType)
+enum class EInterchangeTranslatorAssetType : uint8
+{
+	None = 0,
+	Textures = 1,
+	Materials = 2,
+	Meshes = 3,
+	Animations = 4,
+};
+
 UCLASS(BlueprintType, Blueprintable, Abstract, Experimental)
 class INTERCHANGECORE_API UInterchangeTranslatorBase : public UObject
 {
@@ -31,6 +41,7 @@ public:
 
 	/** Specifies the capabilities of a translator. */
 	virtual EInterchangeTranslatorType GetTranslatorType() const { return EInterchangeTranslatorType::Assets; }
+	virtual bool DoesSupportAssetType(EInterchangeTranslatorAssetType AssetType) const { return false; }
 
 	/** List of formats supported by the translator. Each entry is of the form "ext;Description" where ext is the file extension. */
 	virtual TArray<FString> GetSupportedFormats() const PURE_VIRTUAL(UInterchangeTranslatorBase::GetSupportedExtensions, return TArray<FString>(););
diff --git a/Engine/Source/Runtime/Experimental/Interchange/Engine/Private/InterchangeManager.cpp b/Engine/Source/Runtime/Experimental/Interchange/Engine/Private/InterchangeManager.cpp
--- a/Engine/Source/Runtime/Experimental/Interchange/Engine/Private/InterchangeManager.cpp
+++ b/Engine/Source/Runtime/Experimental/Interchange/Engine/Private/InterchangeManager.cpp
@@ -474,6 +474,23 @@ TArray<FString> UInterchangeManager::GetSupportedFormats(const EInterchangeTrans
 	return FileExtensions;
 }
 
+TArray<FString> UInterchangeManager::GetSupportedAssetTypeFormats(const EInterchangeTranslatorAssetType ForTranslatorAssetType) const
+{
+	TArray<FString> FileExtensions;
+
+	for (const UClass* TranslatorClass : RegisteredTranslatorsClass)
+	{
+		const UInterchangeTranslatorBase* TranslatorBaseCDO = TranslatorClass->GetDefaultObject<UInterchangeTranslatorBase>();
+
+		if (TranslatorBaseCDO->DoesSupportAssetType(ForTranslatorAssetType))
+		{
+			FileExtensions.Append(TranslatorBaseCDO->GetSupportedFormats());
+		}
+	}
+
+	return FileExtensions;
+}
+
 bool UInterchangeManager::CanTranslateSourceData(const UInterchangeSourceData* SourceData) const
 {
 	return GetTranslatorForSourceData(SourceData) != nullptr;
diff --git a/Engine/Source/Runtime/Experimental/Interchange/Engine/Public/InterchangeFilePickerBase.h b/Engine/Source/Runtime/Experimental/Interchange/Engine/Public/InterchangeFilePickerBase.h
new file mode 100644
--- /dev/null
+++ b/Engine/Source/Runtime/Experimental/Interchange/Engine/Public/InterchangeFilePickerBase.h
@@ -0,0 +1,64 @@
+// Copyright Epic Games, Inc. All Rights Reserved.
+
+#pragma once
+
+#include "Containers/UnrealString.h"
+#include "InterchangeTranslatorBase.h"
+#include "UObject/Object.h"
+#include "UObject/ObjectMacros.h"
+
+#include "InterchangeFilePickerBase.generated.h"
+
+UCLASS(Abstract, BlueprintType, Blueprintable, Experimental)
+class INTERCHANGEENGINE_API UInterchangeFilePickerBase : public UObject
+{
+	GENERATED_BODY()
+
+public:
+
+	/**
+	 * Non virtual helper to allow blueprint to implement event base function to implement the function FilePickerForTranslatorAssetType,
+	 */
+	UFUNCTION(BlueprintNativeEvent, BlueprintCallable, Category = "Interchange | Utilities")
+	bool ScriptedFilePickerForTranslatorAssetType(const EInterchangeTranslatorAssetType TranslatorAssetType, FString& OutFilename);
+
+	/** The default implementation (call if the blueprint do not have any implementation) will call the virtual FilePickerForTranslatorAssetType */
+	bool ScriptedFilePickerForTranslatorAssetType_Implementation(const EInterchangeTranslatorAssetType TranslatorAssetType, FString& OutFilename)
+	{
+		//By default we call the virtual function FilePickerForTranslatorAssetType
+		return FilePickerForTranslatorAssetType(TranslatorAssetType, OutFilename);
+	}
+
+	/**
+	 * Non virtual helper to allow blueprint to implement event base function to implement the function FilePickerForTranslatorType,
+	 */
+	UFUNCTION(BlueprintNativeEvent, BlueprintCallable, Category = "Interchange | Utilities")
+	bool ScriptedFilePickerForTranslatorType(const EInterchangeTranslatorType TranslatorType, FString& OutFilename);
+
+	/** The default implementation (call if the blueprint do not have any implementation) will call the virtual FilePickerForTranslatorType */
+	bool ScriptedFilePickerForTranslatorType_Implementation(const EInterchangeTranslatorType TranslatorAssetType, FString& OutFilename)
+	{
+		//By default we call the virtual function FilePickerForTranslatorType
+		return FilePickerForTranslatorType(TranslatorAssetType, OutFilename);
+	}
+
+
+protected:
+
+	/**
+	 * This function must set OutFilename with a valid file that can be import to create the specified TranslatorAssetType.
+	 * We expect the user to see a file picker dialog when the editor is available.
+	 * 
+	 * @Return - true if it set OutFilename with a valid file path, false otherwise.
+	 */
+	virtual bool FilePickerForTranslatorAssetType(const EInterchangeTranslatorAssetType TranslatorAssetType, FString& OutFilename) { return false; }
+
+	/**
+	 * This function must set OutFilename with a valid file that can be import to create the specified TranslatorType.
+	 * We expect the user to see a file picker dialog when the editor is available.
+	 *
+	 * @Return - true if it set OutFilename with a valid file path, false otherwise.
+	 */
+	virtual bool FilePickerForTranslatorType(const EInterchangeTranslatorType TranslatorType, FString& OutFilename) { return false; }
+
+};
diff --git a/Engine/Source/Runtime/Experimental/Interchange/Engine/Public/InterchangeManager.h b/Engine/Source/Runtime/Experimental/Interchange/Engine/Public/InterchangeManager.h
--- a/Engine/Source/Runtime/Experimental/Interchange/Engine/Public/InterchangeManager.h
+++ b/Engine/Source/Runtime/Experimental/Interchange/Engine/Public/InterchangeManager.h
@@ -350,6 +350,11 @@ public:
 	 */
 	TArray<FString> GetSupportedFormats(const EInterchangeTranslatorType ForTranslatorType) const;
 
+	/**
+	 * Returns the list of formats supporting the specified translator asset type.
+	 */
+	TArray<FString> GetSupportedAssetTypeFormats(const EInterchangeTranslatorAssetType ForTranslatorAssetType) const;
+
 	/**
 	 * Look if there is a registered translator for this source data.
 	 * This allow us to by pass the original asset tools system to import supported asset.
diff --git a/Engine/Source/Runtime/Experimental/Interchange/Engine/Public/InterchangeProjectSettings.h b/Engine/Source/Runtime/Experimental/Interchange/Engine/Public/InterchangeProjectSettings.h
--- a/Engine/Source/Runtime/Experimental/Interchange/Engine/Public/InterchangeProjectSettings.h
+++ b/Engine/Source/Runtime/Experimental/Interchange/Engine/Public/InterchangeProjectSettings.h
@@ -5,6 +5,7 @@
 #include "CoreMinimal.h"
 #include "Engine/DeveloperSettings.h"
 #include "Engine/EngineTypes.h"
+#include "InterchangeFilePickerBase.h"
 #include "InterchangePipelineBase.h"
 #include "InterchangePipelineConfigurationBase.h"
 
@@ -42,6 +43,10 @@ public:
 	/** If enabled, the pipeline stacks configuration dialog will show every time interchange must choose a pipeline to import or re-import. If disabled interchange will use the DefaultPipelineStack.*/
 	UPROPERTY(EditAnywhere, config, Category = Interchange)
 	bool bShowPipelineStacksConfigurationDialog;
+
+	/** This tells interchange which file picker class to construct when we need to choose a file for a source.*/
+	UPROPERTY(EditAnywhere, config, Category = Interchange)
+	TSoftClassPtr <UInterchangeFilePickerBase> FilePickerClass;
 #endif
 
 	/** If checked, will use Interchange when importing into level.*/
