Fix crash when using the interchange graph inspector
#rnx

diff --git a/Engine/Plugins/Experimental/Interchange/Editor/Source/Pipelines/Private/SInterchangeGraphInspectorWindow.cpp b/Engine/Plugins/Experimental/Interchange/Editor/Source/Pipelines/Private/SInterchangeGraphInspectorWindow.cpp
--- a/Engine/Plugins/Experimental/Interchange/Editor/Source/Pipelines/Private/SInterchangeGraphInspectorWindow.cpp
+++ b/Engine/Plugins/Experimental/Interchange/Editor/Source/Pipelines/Private/SInterchangeGraphInspectorWindow.cpp
@@ -104,9 +104,8 @@ public:
 			if (UClass* IconClass = FactoryNode->GetObjectClass())
 			{
 				TypeIcon = FSlateIconFinder::FindIconBrushForClass(IconClass);
+				Tooltip += TEXT(" [") + IconClass->GetName() + TEXT("]");
 			}
-
-			Tooltip += TEXT(" [") + FactoryNode->GetObjectClass()->GetName() + TEXT("]");
 		}
 
 		if (!TypeIcon)
