Interchange: In the interchange pipeline configuration dialog, add a tooltip describing why the import button is greyout.

#rnx

diff --git a/Engine/Plugins/Experimental/Interchange/Editor/Source/Pipelines/Private/SInterchangePipelineConfigurationDialog.cpp b/Engine/Plugins/Experimental/Interchange/Editor/Source/Pipelines/Private/SInterchangePipelineConfigurationDialog.cpp
--- a/Engine/Plugins/Experimental/Interchange/Editor/Source/Pipelines/Private/SInterchangePipelineConfigurationDialog.cpp
+++ b/Engine/Plugins/Experimental/Interchange/Editor/Source/Pipelines/Private/SInterchangePipelineConfigurationDialog.cpp
@@ -457,6 +457,7 @@ void SInterchangePipelineConfigurationDialog::Construct(const FArguments& InArgs
 					SNew(SButton)
 					.HAlign(HAlign_Center)
 					.Text(LOCTEXT("InspectorGraphWindow_ImportAll", "Import All"))
+					.ToolTipText(this, &SInterchangePipelineConfigurationDialog::GetImportButtonTooltip)
 					.IsEnabled(this, &SInterchangePipelineConfigurationDialog::IsImportButtonEnabled)
 					.OnClicked(this, &SInterchangePipelineConfigurationDialog::OnCloseDialog, ECloseEventType::ImportAll)
 				]
@@ -465,6 +466,7 @@ void SInterchangePipelineConfigurationDialog::Construct(const FArguments& InArgs
 					SNew(SButton)
 					.HAlign(HAlign_Center)
 					.Text(LOCTEXT("InspectorGraphWindow_Import", "Import"))
+					.ToolTipText(this, &SInterchangePipelineConfigurationDialog::GetImportButtonTooltip)
 					.IsEnabled(this, &SInterchangePipelineConfigurationDialog::IsImportButtonEnabled)
 					.OnClicked(this, &SInterchangePipelineConfigurationDialog::OnCloseDialog, ECloseEventType::Import)
 				]
@@ -498,18 +500,18 @@ FText SInterchangePipelineConfigurationDialog::GetSourceDescription() const
 	return ActionDescription;
 }
 
-bool SInterchangePipelineConfigurationDialog::RecursiveValidatePipelineSettings(const TSharedPtr<FInterchangePipelineStacksTreeNodeItem>& ParentNode) const
+bool SInterchangePipelineConfigurationDialog::RecursiveValidatePipelineSettings(const TSharedPtr<FInterchangePipelineStacksTreeNodeItem>& ParentNode, TOptional<FText>& OutInvalidReason) const
 {
 	if (ParentNode->Pipeline)
 	{
-		if (!ParentNode->Pipeline->IsSettingsAreValid())
+		if (!ParentNode->Pipeline->IsSettingsAreValid(OutInvalidReason))
 		{
 			return false;
 		}
 	}
 	for (int32 ChildIndex = 0; ChildIndex < ParentNode->Childrens.Num(); ++ChildIndex)
 	{
-		if (!RecursiveValidatePipelineSettings(ParentNode->Childrens[ChildIndex]))
+		if (!RecursiveValidatePipelineSettings(ParentNode->Childrens[ChildIndex], OutInvalidReason))
 		{
 			return false;
 		}
@@ -526,13 +528,32 @@ bool SInterchangePipelineConfigurationDialog::IsImportButtonEnabled() const
 		//Reimport have only one stack and the name wont match DefaultPipelineStackName
 		if (bReimport || RootNode->StackName == DefaultPipelineStackName)
 		{
-			return RecursiveValidatePipelineSettings(RootNode);
+			TOptional<FText> InvalidReason;
+			return RecursiveValidatePipelineSettings(RootNode, InvalidReason);
 		}
 	}
-
 	return true;
 }
 
+FText SInterchangePipelineConfigurationDialog::GetImportButtonTooltip() const
+{
+	const FName DefaultPipelineStackName = GetDefault<UInterchangeProjectSettings>()->DefaultPipelineStack;
+	const TArray<TSharedPtr<FInterchangePipelineStacksTreeNodeItem>>& RootNodeArray = PipelineConfigurationTreeView->GetRootNodeArray();
+	for (const TSharedPtr<FInterchangePipelineStacksTreeNodeItem>& RootNode : RootNodeArray)
+	{
+		//Reimport have only one stack and the name wont match DefaultPipelineStackName
+		if (bReimport || RootNode->StackName == DefaultPipelineStackName)
+		{
+			TOptional<FText> InvalidReason;
+			if (!RecursiveValidatePipelineSettings(RootNode, InvalidReason) && InvalidReason.IsSet())
+			{
+				return InvalidReason.GetValue();
+			}
+		}
+	}
+	return FText();
+}
+
 void SInterchangePipelineConfigurationDialog::RecursiveSavePipelineSettings(const TSharedPtr<FInterchangePipelineStacksTreeNodeItem>& ParentNode, const int32 PipelineIndex) const
 {
 	if (ParentNode->Pipeline)
diff --git a/Engine/Plugins/Experimental/Interchange/Editor/Source/Pipelines/Private/SInterchangePipelineConfigurationDialog.h b/Engine/Plugins/Experimental/Interchange/Editor/Source/Pipelines/Private/SInterchangePipelineConfigurationDialog.h
--- a/Engine/Plugins/Experimental/Interchange/Editor/Source/Pipelines/Private/SInterchangePipelineConfigurationDialog.h
+++ b/Engine/Plugins/Experimental/Interchange/Editor/Source/Pipelines/Private/SInterchangePipelineConfigurationDialog.h
@@ -132,8 +132,9 @@ private:
 
 	FText GetSourceDescription() const;
 
-	bool RecursiveValidatePipelineSettings(const TSharedPtr<FInterchangePipelineStacksTreeNodeItem>& ParentNode) const;
+	bool RecursiveValidatePipelineSettings(const TSharedPtr<FInterchangePipelineStacksTreeNodeItem>& ParentNode, TOptional<FText>& OutInvalidReason) const;
 	bool IsImportButtonEnabled() const;
+	FText GetImportButtonTooltip() const;
 
 	void RecursiveSavePipelineSettings(const TSharedPtr<FInterchangePipelineStacksTreeNodeItem>& ParentNode, const int32 PipelineIndex) const;
 	void RecursiveLoadPipelineSettings(const TSharedPtr<FInterchangePipelineStacksTreeNodeItem>& ParentNode, const int32 PipelineIndex) const;
diff --git a/Engine/Plugins/Experimental/Interchange/Runtime/Source/Pipelines/Private/InterchangeGenericAssetsPipeline.cpp b/Engine/Plugins/Experimental/Interchange/Runtime/Source/Pipelines/Private/InterchangeGenericAssetsPipeline.cpp
--- a/Engine/Plugins/Experimental/Interchange/Runtime/Source/Pipelines/Private/InterchangeGenericAssetsPipeline.cpp
+++ b/Engine/Plugins/Experimental/Interchange/Runtime/Source/Pipelines/Private/InterchangeGenericAssetsPipeline.cpp
@@ -72,39 +72,39 @@ void UInterchangeGenericAssetsPipeline::PreDialogCleanup(const FName PipelineSta
 	SaveSettings(PipelineStackName);
 }
 
-bool UInterchangeGenericAssetsPipeline::IsSettingsAreValid() const
+bool UInterchangeGenericAssetsPipeline::IsSettingsAreValid(TOptional<FText>& OutInvalidReason) const
 {
-	if (TexturePipeline && !TexturePipeline->IsSettingsAreValid())
+	if (TexturePipeline && !TexturePipeline->IsSettingsAreValid(OutInvalidReason))
 	{
 		return false;
 	}
 
-	if (MaterialPipeline && !MaterialPipeline->IsSettingsAreValid())
+	if (MaterialPipeline && !MaterialPipeline->IsSettingsAreValid(OutInvalidReason))
 	{
 		return false;
 	}
 
-	if (CommonMeshesProperties && !CommonMeshesProperties->IsSettingsAreValid())
+	if (CommonMeshesProperties && !CommonMeshesProperties->IsSettingsAreValid(OutInvalidReason))
 	{
 		return false;
 	}
 
-	if (CommonSkeletalMeshesAndAnimationsProperties && !CommonSkeletalMeshesAndAnimationsProperties->IsSettingsAreValid())
+	if (CommonSkeletalMeshesAndAnimationsProperties && !CommonSkeletalMeshesAndAnimationsProperties->IsSettingsAreValid(OutInvalidReason))
 	{
 		return false;
 	}
 
-	if (MeshPipeline && !MeshPipeline->IsSettingsAreValid())
+	if (MeshPipeline && !MeshPipeline->IsSettingsAreValid(OutInvalidReason))
 	{
 		return false;
 	}
 
-	if (AnimationPipeline && !AnimationPipeline->IsSettingsAreValid())
+	if (AnimationPipeline && !AnimationPipeline->IsSettingsAreValid(OutInvalidReason))
 	{
 		return false;
 	}
 
-	return Super::IsSettingsAreValid();
+	return Super::IsSettingsAreValid(OutInvalidReason);
 }
 
 void UInterchangeGenericAssetsPipeline::AdjustSettingsForReimportType(EInterchangeReimportType ImportType, TObjectPtr<UObject> ReimportAsset)
diff --git a/Engine/Plugins/Experimental/Interchange/Runtime/Source/Pipelines/Public/InterchangeGenericAssetsPipeline.h b/Engine/Plugins/Experimental/Interchange/Runtime/Source/Pipelines/Public/InterchangeGenericAssetsPipeline.h
--- a/Engine/Plugins/Experimental/Interchange/Runtime/Source/Pipelines/Public/InterchangeGenericAssetsPipeline.h
+++ b/Engine/Plugins/Experimental/Interchange/Runtime/Source/Pipelines/Public/InterchangeGenericAssetsPipeline.h
@@ -78,7 +78,7 @@ public:
 
 	virtual void PreDialogCleanup(const FName PipelineStackName) override;
 
-	virtual bool IsSettingsAreValid() const override;
+	virtual bool IsSettingsAreValid(TOptional<FText>& OutInvalidReason) const override;
 
 
 	virtual void AdjustSettingsForReimportType(EInterchangeReimportType ImportType, TObjectPtr<UObject> ReimportAsset) override;
diff --git a/Engine/Plugins/Experimental/Interchange/Runtime/Source/Pipelines/Public/InterchangeGenericAssetsPipelineSharedSettings.h b/Engine/Plugins/Experimental/Interchange/Runtime/Source/Pipelines/Public/InterchangeGenericAssetsPipelineSharedSettings.h
--- a/Engine/Plugins/Experimental/Interchange/Runtime/Source/Pipelines/Public/InterchangeGenericAssetsPipelineSharedSettings.h
+++ b/Engine/Plugins/Experimental/Interchange/Runtime/Source/Pipelines/Public/InterchangeGenericAssetsPipelineSharedSettings.h
@@ -93,12 +93,13 @@ public:
 	UPROPERTY(EditAnywhere, BlueprintReadWrite, Category = "Common Skeletal Meshes and Animations")
 	bool bUseT0AsRefPose = false;
 
-	virtual bool IsSettingsAreValid() const override
+	virtual bool IsSettingsAreValid(TOptional<FText>& OutInvalidReason) const override
 	{
 		if (bImportOnlyAnimations && Skeleton.IsNull())
 		{
+			OutInvalidReason = NSLOCTEXT("UInterchangeGenericCommonSkeletalMeshesAndAnimationsProperties", "SkeletonMustBeSpecified", "When importing only animations, a valid skeleton must be set.");
 			return false;
 		}
-		return Super::IsSettingsAreValid();
+		return Super::IsSettingsAreValid(OutInvalidReason);
 	}
 };
diff --git a/Engine/Source/Runtime/Experimental/Interchange/Core/Public/InterchangePipelineBase.h b/Engine/Source/Runtime/Experimental/Interchange/Core/Public/InterchangePipelineBase.h
--- a/Engine/Source/Runtime/Experimental/Interchange/Core/Public/InterchangePipelineBase.h
+++ b/Engine/Source/Runtime/Experimental/Interchange/Core/Public/InterchangePipelineBase.h
@@ -133,7 +133,7 @@ public:
 	 * This function should return true if all the pipeline settings are in a valid state to start the import.
 	 * The Interchange Pipeline Configuration Dialog will call this to know if the Import button can be enable.
 	 */
-	virtual bool IsSettingsAreValid() const
+	virtual bool IsSettingsAreValid(TOptional<FText>& OutInvalidReason) const
 	{
 		return true;
 	}
