Interchange: Implement a new workflow for the pipeline stack.

The project settings pipeline stack now point on assets (base on UInterchangePipelineBase) instead of class. The new UInterchangePipeline asset editor allow users to specificly lock any property values. The import dialog show the pipeline locked properties has read only.

In the import dialog, a "Reset to default" button was added, when users click on it, all selected pipeline properties value will go back to the pipeline statck asset properties value. In case of re-import the properties value will get back to the last import properties value.

#rnx

diff --git a/Engine/Config/BaseEngine.ini b/Engine/Config/BaseEngine.ini
--- a/Engine/Config/BaseEngine.ini
+++ b/Engine/Config/BaseEngine.ini
@@ -2995,14 +2995,6 @@ DefaultChaosSolverActorClass=/Script/ChaosSolverEngine.ChaosSolverActor
 [PlatformPaths]
 UnrealBuildTool=Engine/Build/BatchFiles/Build.bat
 
-[/Script/InterchangeEngine.InterchangeProjectSettings]
-PipelineStacks=(("Generic", (Pipelines=("/Script/InterchangePipelines.InterchangeGenericAssetsPipeline"))),("Scene", (Pipelines=("/Script/InterchangePipelines.InterchangeGenericAssetsPipeline","/Script/InterchangePipelines.InterchangeGenericLevelPipeline"))))
-DefaultPipelineStack=Generic
-DefaultScenePipelineStack=Scene
-PipelineConfigurationDialogClass=/Script/InterchangeEditorPipelines.InterchangePipelineConfigurationGeneric
-bShowPipelineStacksConfigurationDialog=True
-FilePickerClass=/Script/InterchangeEditorUtilities.InterchangeFilePickerGeneric
-
 [AlternateTextureCompression]
 TextureCompressionFormat="TextureFormatOodle"
 
@@ -3047,3 +3039,13 @@ EnvironmentCollisionThickness = 2.f
 ; Limit the number of extracted shapes per component as these collisions are very expensive
 EnvironmentCollisionMaxShapes = 32
 
+[/Script/InterchangeEngine.InterchangeProjectSettings]
+PipelineStacks=(("Assets", (Pipelines=("/Interchange/Pipelines/DefaultAssetsPipeline.DefaultAssetsPipeline"))),("Scene", (Pipelines=("/Interchange/Pipelines/DefaultSceneAssetsPipeline.DefaultSceneAssetsPipeline","/Interchange/Pipelines/DefaultSceneLevelPipeline.DefaultSceneLevelPipeline"))))
+DefaultPipelineStack=Assets
+DefaultScenePipelineStack=Scene
+PipelineConfigurationDialogClass=/Script/InterchangeEditorPipelines.InterchangePipelineConfigurationGeneric
+bShowPipelineStacksConfigurationDialog=True
+FilePickerClass=/Script/InterchangeEditorUtilities.InterchangeFilePickerGeneric
+bUseInterchangeWhenImportingIntoLevel=False
+bStaticMeshUseSmoothEdgesIfSmoothingInformationIsMissing=True
+
diff --git a/Engine/Plugins/Experimental/Interchange/Editor/Source/Pipelines/InterchangeEditorPipelines.Build.cs b/Engine/Plugins/Experimental/Interchange/Editor/Source/Pipelines/InterchangeEditorPipelines.Build.cs
--- a/Engine/Plugins/Experimental/Interchange/Editor/Source/Pipelines/InterchangeEditorPipelines.Build.cs
+++ b/Engine/Plugins/Experimental/Interchange/Editor/Source/Pipelines/InterchangeEditorPipelines.Build.cs
@@ -21,11 +21,14 @@ namespace UnrealBuildTool.Rules
 				new string[]
 				{
 					"InputCore",
+					"InterchangePipelines",
 					"MainFrame",
 					"Projects",
 					"PropertyEditor",
 					"Slate",
 					"SlateCore",
+					"UnrealEd",
+					"AssetTools"
 				}
 			);
 		}
diff --git a/Engine/Plugins/Experimental/Interchange/Editor/Source/Pipelines/Private/InterchangeEditorPipelineDetails.cpp b/Engine/Plugins/Experimental/Interchange/Editor/Source/Pipelines/Private/InterchangeEditorPipelineDetails.cpp
--- a/Engine/Plugins/Experimental/Interchange/Editor/Source/Pipelines/Private/InterchangeEditorPipelineDetails.cpp
+++ b/Engine/Plugins/Experimental/Interchange/Editor/Source/Pipelines/Private/InterchangeEditorPipelineDetails.cpp
@@ -7,8 +7,12 @@
 #include "DetailWidgetRow.h"
 #include "Editor.h"
 #include "IDetailGroup.h"
+#include "InterchangePipelineBase.h"
 #include "Nodes/InterchangeBaseNode.h"
 #include "PropertyHandle.h"
+#include "ScopedTransaction.h"
+#include "Styling/AppStyle.h"
+#include "Styling/StyleColors.h"
 #include "Widgets/DeclarativeSyntaxSupport.h"
 #include "Widgets/Images/SImage.h"
 #include "Widgets/Input/SCheckBox.h"
@@ -17,9 +21,261 @@
 #include "Widgets/Text/STextBlock.h"
 #include "Widgets/SBoxPanel.h"
 
-
 #define LOCTEXT_NAMESPACE "InterchangeEditorPipelineDetails"
 
+FInterchangePipelineBaseDetailsCustomization::FInterchangePipelineBaseDetailsCustomization()
+{
+	InterchangePipeline = nullptr;
+	CachedDetailBuilder = nullptr;
+}
+
+void FInterchangePipelineBaseDetailsCustomization::RefreshCustomDetail()
+{
+	if (CachedDetailBuilder)
+	{
+		CachedDetailBuilder->ForceRefreshDetails();
+	}
+}
+
+void FInterchangePipelineBaseDetailsCustomization::LockPropertyHandleRow(const TSharedPtr<IPropertyHandle> PropertyHandle, IDetailPropertyRow& PropertyRow) const
+{
+	//When we have a locked property we do not show the Reset to default button
+	FIsResetToDefaultVisible IsResetVisible = FIsResetToDefaultVisible::CreateLambda([](TSharedPtr<IPropertyHandle> PropertyHandle)
+		{
+			return false;
+		});
+	FResetToDefaultHandler ResetHandler = FResetToDefaultHandler::CreateLambda([](TSharedPtr<IPropertyHandle> PropertyHandle) { return; });
+	FResetToDefaultOverride ResetOverride = FResetToDefaultOverride::Create(IsResetVisible, ResetHandler);
+
+	uint32 NumChildren = 0;
+	if (PropertyHandle.Get()->GetNumChildren(NumChildren) == FPropertyAccess::Success && NumChildren > 0)
+	{
+		//Remove the reset to CDO for the locked property
+		constexpr bool bShowChildren = false;
+		TSharedPtr<SWidget> NameWidget;
+		TSharedPtr<SWidget> ValueWidget;
+		FDetailWidgetRow Row;
+		PropertyRow.GetDefaultWidgets(NameWidget, ValueWidget, Row);
+		NameWidget->SetEnabled(false);
+		ValueWidget->SetEnabled(false);
+		FDetailWidgetRow& WidgetRow = PropertyRow.CustomWidget(bShowChildren)
+			.NameContent()
+			.MinDesiredWidth(Row.NameWidget.MinWidth)
+			.MaxDesiredWidth(Row.NameWidget.MaxWidth)
+			[
+				NameWidget.ToSharedRef()
+			]
+		.ValueContent()
+			.MinDesiredWidth(Row.ValueWidget.MinWidth)
+			.MaxDesiredWidth(Row.ValueWidget.MaxWidth)
+			.VAlign(VAlign_Center)
+			[
+				ValueWidget.ToSharedRef()
+			];
+		WidgetRow.OverrideResetToDefault(ResetOverride);
+	}
+	else
+	{
+		PropertyRow.IsEnabled(false);
+		PropertyRow.OverrideResetToDefault(ResetOverride);
+	}
+}
+
+void FInterchangePipelineBaseDetailsCustomization::AddSubCategory(IDetailLayoutBuilder& DetailBuilder, TMap<FName, TMap<FName, TArray<FInternalPropertyData>>>& SubCategoriesPropertiesPerMainCategory)
+{
+	for (const TPair<FName, TMap<FName, TArray<FInternalPropertyData>>>& MainCategoryAndSubCategoriesProperties : SubCategoriesPropertiesPerMainCategory)
+	{
+		const FName& MainCategoryName = MainCategoryAndSubCategoriesProperties.Key;
+		const TMap<FName, TArray<FInternalPropertyData>>& SubCategoriesProperties = MainCategoryAndSubCategoriesProperties.Value;
+		IDetailCategoryBuilder& MainCategory = DetailBuilder.EditCategory(MainCategoryName);
+		//If we found some sub category we can add them to the group
+		for (const TPair<FName, TArray<FInternalPropertyData>>& SubCategoryProperties : SubCategoriesProperties)
+		{
+			const FName& SubCategoryName = SubCategoryProperties.Key;
+			const TArray<FInternalPropertyData>& Properties = SubCategoryProperties.Value;
+			constexpr bool SubCategoryAdvanced = false;
+			IDetailGroup& Group = MainCategory.AddGroup(SubCategoryName, FText::FromName(SubCategoryName), SubCategoryAdvanced);
+			for (int32 PropertyIndex = 0; PropertyIndex < Properties.Num(); ++PropertyIndex)
+			{
+				const FInternalPropertyData& PropertyData = Properties[PropertyIndex];
+				const TSharedPtr<IPropertyHandle>& PropertyHandle = PropertyData.PropertyHandle;
+				DetailBuilder.HideProperty(PropertyHandle);
+				IDetailPropertyRow& PropertyRow = Group.AddPropertyRow(PropertyHandle.ToSharedRef());
+				//If needed, make the property read only
+				if (PropertyData.bReadOnly)
+				{
+					LockPropertyHandleRow(PropertyHandle, PropertyRow);
+				}
+			}
+		}
+	}
+}
+
+TSharedRef<IDetailCustomization> FInterchangePipelineBaseDetailsCustomization::MakeInstance()
+{
+	return MakeShareable(new FInterchangePipelineBaseDetailsCustomization);
+}
+
+void FInterchangePipelineBaseDetailsCustomization::CustomizeDetails(IDetailLayoutBuilder& DetailBuilder)
+{
+	CachedDetailBuilder = &DetailBuilder;
+	TArray<TWeakObjectPtr<UObject>> EditingObjects;
+	DetailBuilder.GetObjectsBeingCustomized(EditingObjects);
+	check(EditingObjects.Num() == 1);
+
+	InterchangePipeline = Cast<UInterchangePipelineBase>(EditingObjects[0].Get());
+
+	if (!ensure(InterchangePipeline.IsValid()))
+	{
+		return;
+	}
+
+	const bool bAllowLockedPropertiesEdition = InterchangePipeline->bAllowLockedPropertiesEdition;
+
+	TArray<FName> AllCategoryNames;
+	CachedDetailBuilder->GetCategoryNames(AllCategoryNames);
+	TMap<FName, TArray<FName>> PropertiesPerCategorys;
+	InternalGetPipelineProperties(InterchangePipeline.Get(), AllCategoryNames, PropertiesPerCategorys);
+
+	TMap<FName, TMap<FName, TArray<FInternalPropertyData>>> SubCategoriesPropertiesPerMainCategory;
+
+	for (const TPair<FName, TArray<FName>>& CategoryAndProperties : PropertiesPerCategorys)
+	{
+		const FName CategoryName = CategoryAndProperties.Key;
+		IDetailCategoryBuilder& Category = CachedDetailBuilder->EditCategory(CategoryName);
+		
+		TArray<TSharedRef<IPropertyHandle>> CategoryProperties;
+
+		Category.GetDefaultProperties(CategoryProperties);
+
+		for (TSharedRef<IPropertyHandle>& PropertyHandle : CategoryProperties)
+		{
+			FProperty* PropertyPtr = PropertyHandle.Get().GetProperty();
+			if (!PropertyPtr)
+			{
+				continue;
+			}
+			const FName PropertyName = PropertyHandle.Get().GetProperty() ? PropertyHandle.Get().GetProperty()->GetFName() : NAME_None;
+			if (PropertyName == UInterchangePipelineBase::GetLockedPropertiesPropertyName())
+			{
+				CachedDetailBuilder->HideProperty(PropertyHandle);
+				continue;
+			}
+			// Skip the property if:
+			// We do not allow property edition
+			// property is not in the List of supported properties
+			if (!CategoryAndProperties.Value.Contains(PropertyName))
+			{
+				continue;
+			}
+			FName PropertyPath = FName(PropertyPtr->GetPathName());
+			CachedDetailBuilder->HideProperty(PropertyHandle);
+			if (!bAllowLockedPropertiesEdition)
+			{
+				const bool bLockedProperty = InterchangePipeline->GetLockedPropertyStatus(PropertyPath);
+				const FName SubCategoryData = FName(PropertyHandle->GetMetaData(TEXT("SubCategory")));
+				if (SubCategoryData != NAME_None)
+				{
+					TMap<FName, TArray<FInternalPropertyData>>& SubCategoriesProperties = SubCategoriesPropertiesPerMainCategory.FindOrAdd(CategoryName);
+					TArray<FInternalPropertyData>& SubCategoryProperties = SubCategoriesProperties.FindOrAdd(SubCategoryData);
+					FInternalPropertyData& PropertyData = SubCategoryProperties.AddDefaulted_GetRef();
+					PropertyData.PropertyHandle = PropertyHandle;
+					PropertyData.bReadOnly = bLockedProperty;
+				}
+				else
+				{
+					IDetailPropertyRow& PropertyRow = Category.AddProperty(PropertyHandle);
+					//When we use the pipeline in interchange
+					if (bLockedProperty)
+					{
+						LockPropertyHandleRow(PropertyHandle, PropertyRow);
+					}
+				}
+			}
+			else
+			{
+				constexpr bool bShowChildren = true;
+				IDetailPropertyRow& PropertyRow = Category.AddProperty(PropertyHandle);
+				TSharedPtr<SWidget> NameWidget;
+				TSharedPtr<SWidget> ValueWidget;
+				FDetailWidgetRow Row;
+				PropertyRow.GetDefaultWidgets(NameWidget, ValueWidget, Row);
+				//When we edit the default value and the locked properties of a pipeline assets
+				PropertyRow.CustomWidget(bShowChildren)
+				.NameContent()
+				.MinDesiredWidth(Row.NameWidget.MinWidth)
+				.MaxDesiredWidth(Row.NameWidget.MaxWidth)
+				[
+					SNew(SHorizontalBox)
+					+ SHorizontalBox::Slot()
+					.AutoWidth()
+					.Padding(3.0f, 1.0f)
+					[
+						SNew(SCheckBox)
+						.CheckedImage(FAppStyle::Get().GetBrush("Icons.Lock"))
+						.CheckedHoveredImage(FAppStyle::Get().GetBrush("Icons.Lock"))
+						.CheckedPressedImage(FAppStyle::Get().GetBrush("Icons.Lock"))
+						.UncheckedImage(FAppStyle::Get().GetBrush("Icons.Unlock"))
+						.UncheckedHoveredImage(FAppStyle::Get().GetBrush("Icons.Unlock"))
+						.UncheckedPressedImage(FAppStyle::Get().GetBrush("Icons.Unlock"))
+						.ToolTipText(NSLOCTEXT("InterchangePipelineBaseDetails::CustomizeDetails", "LockedTooltip", "If true this property will be readonly in the interchange import dialog."))
+						.OnCheckStateChanged_Lambda([InterchangePipelinePtr = InterchangePipeline, PropertyPath](ECheckBoxState CheckType)
+						{
+							if (!ensure(InterchangePipelinePtr.IsValid()))
+							{
+								return;
+							}
+							FScopedTransaction ScopedTransaction(NSLOCTEXT("InterchangePipelineBaseDetails::CustomizeDetails", "TransactionLockedPropertiesSet", "Toggle Property Locked"), !GIsTransacting);
+							InterchangePipelinePtr->Modify();
+							InterchangePipelinePtr->SetLockedPropertyStatus(PropertyPath, (CheckType == ECheckBoxState::Checked));
+							InterchangePipelinePtr->PostEditChange();
+						})
+						.IsChecked_Lambda([InterchangePipelinePtr = InterchangePipeline, PropertyPath]()
+						{
+							if (!InterchangePipelinePtr.IsValid())
+							{
+								return ECheckBoxState::Unchecked;
+							}
+							return InterchangePipelinePtr->GetLockedPropertyStatus(PropertyPath) ? ECheckBoxState::Checked : ECheckBoxState::Unchecked;
+						})
+					]
+					+ SHorizontalBox::Slot()
+					.FillWidth(1.0f)
+					.Padding(3.0f, 1.0f)
+					[
+						SNew(SBorder)
+						.BorderBackgroundColor(FStyleColors::Background)
+						.ForegroundColor_Lambda([InterchangePipelinePtr = InterchangePipeline, PropertyPath]()
+							{
+								if (!InterchangePipelinePtr.IsValid())
+								{
+									return FStyleColors::Foreground;
+								}
+								if (InterchangePipelinePtr->GetLockedPropertyStatus(PropertyPath))
+								{
+									return FStyleColors::AccentBlue;
+								}
+								return FStyleColors::Foreground;
+							})
+						[
+							NameWidget.ToSharedRef()
+						]
+					]
+				]
+				.ValueContent()
+				.MinDesiredWidth(Row.ValueWidget.MinWidth)
+				.MaxDesiredWidth(Row.ValueWidget.MaxWidth)
+				.VAlign(VAlign_Center)
+				[
+					ValueWidget.ToSharedRef()
+				];
+			}
+		}
+	}
+
+	AddSubCategory(DetailBuilder, SubCategoriesPropertiesPerMainCategory);
+}
+
 FInterchangeBaseNodeDetailsCustomization::FInterchangeBaseNodeDetailsCustomization()
 {
 	InterchangeBaseNode = nullptr;
@@ -38,6 +294,40 @@ void FInterchangeBaseNodeDetailsCustomization::RefreshCustomDetail()
 	}
 }
 
+void FInterchangePipelineBaseDetailsCustomization::InternalGetPipelineProperties(const UInterchangePipelineBase* Pipeline, const TArray<FName>& AllCategoryNames, TMap<FName, TArray<FName>>& PropertiesPerCategorys) const
+{
+	const UClass* PipelineClass = Pipeline->GetClass();
+	const FName CategoryKey("Category");
+	TArray<UInterchangePipelineBase*> SubPipelines;
+	for (FProperty* Property = PipelineClass->PropertyLink; Property; Property = Property->PropertyLinkNext)
+	{
+		//Do not load a transient property
+		if (Property->HasAnyPropertyFlags(CPF_Transient))
+		{
+			continue;
+		}
+		FObjectProperty* SubObject = CastField<FObjectProperty>(Property);
+		if (UInterchangePipelineBase* SubPipeline = SubObject ? Cast<UInterchangePipelineBase>(SubObject->GetObjectPropertyValue_InContainer(Pipeline)) : nullptr)
+		{
+			SubPipelines.Add(SubPipeline);
+		}
+		else if (const FString* PropertyCategoryString = Property->FindMetaData(CategoryKey))
+		{
+			FName PropertyCategoryName(*PropertyCategoryString);
+			if (AllCategoryNames.Contains(PropertyCategoryName))
+			{
+				PropertiesPerCategorys.FindOrAdd(PropertyCategoryName).AddUnique(Property->GetFName());
+			}
+		}
+	}
+	//Get all sub pipeline object properties
+	for (UInterchangePipelineBase* SubPipeline : SubPipelines)
+	{
+		InternalGetPipelineProperties(SubPipeline, AllCategoryNames, PropertiesPerCategorys);
+	}
+}
+
+
 TSharedRef<IDetailCustomization> FInterchangeBaseNodeDetailsCustomization::MakeInstance()
 {
 	return MakeShareable(new FInterchangeBaseNodeDetailsCustomization);
diff --git a/Engine/Plugins/Experimental/Interchange/Editor/Source/Pipelines/Private/InterchangeEditorPipelineDetails.h b/Engine/Plugins/Experimental/Interchange/Editor/Source/Pipelines/Private/InterchangeEditorPipelineDetails.h
--- a/Engine/Plugins/Experimental/Interchange/Editor/Source/Pipelines/Private/InterchangeEditorPipelineDetails.h
+++ b/Engine/Plugins/Experimental/Interchange/Editor/Source/Pipelines/Private/InterchangeEditorPipelineDetails.h
@@ -3,7 +3,7 @@
 #pragma once
 
 #include "CoreMinimal.h"
-#include "UObject/WeakObjectPtr.h"
+#include "UObject/WeakObjectPtrTemplates.h"
 #include "Types/SlateEnums.h"
 #include "IDetailCustomization.h"
 #include "Input/Reply.h"
@@ -13,6 +13,7 @@ class IDetailLayoutBuilder;
 class IDetailPropertyRow;
 class IPropertyHandle;
 class UInterchangeBaseNode;
+class UInterchangePipelineBase;
 
 namespace UE
 {
@@ -22,6 +23,33 @@ namespace UE
 	}
 }
 
+class FInterchangePipelineBaseDetailsCustomization : public IDetailCustomization
+{
+public:
+
+	/** Makes a new instance of this detail layout class for a specific detail view requesting it */
+	static TSharedRef<IDetailCustomization> MakeInstance();
+
+	/** IDetailCustomization interface */
+	virtual void CustomizeDetails(IDetailLayoutBuilder& DetailBuilder) override;
+	/** End IDetailCustomization interface */
+
+private:
+	struct FInternalPropertyData
+	{
+		TSharedPtr<IPropertyHandle> PropertyHandle;
+		bool bReadOnly = false;
+	};
+	void RefreshCustomDetail();
+	void LockPropertyHandleRow(const TSharedPtr<IPropertyHandle> PropertyHandle, IDetailPropertyRow& PropertyRow) const;
+	void AddSubCategory(IDetailLayoutBuilder& DetailBuilder, TMap<FName, TMap<FName, TArray<FInternalPropertyData>>>& SubCategoriesPropertiesPerMainCategory);
+	void InternalGetPipelineProperties(const UInterchangePipelineBase* Pipeline, const TArray<FName>& AllCategoryNames, TMap<FName, TArray<FName>>& PropertiesPerCategorys) const;
+	/** Use MakeInstance to create an instance of this class */
+	FInterchangePipelineBaseDetailsCustomization();
+
+	TWeakObjectPtr<UInterchangePipelineBase> InterchangePipeline;
+	IDetailLayoutBuilder* CachedDetailBuilder;	// The detail builder for this customisation
+};
 
 class FInterchangeBaseNodeDetailsCustomization : public IDetailCustomization
 {
diff --git a/Engine/Plugins/Experimental/Interchange/Editor/Source/Pipelines/Private/InterchangeEditorPipelinesModule.cpp b/Engine/Plugins/Experimental/Interchange/Editor/Source/Pipelines/Private/InterchangeEditorPipelinesModule.cpp
--- a/Engine/Plugins/Experimental/Interchange/Editor/Source/Pipelines/Private/InterchangeEditorPipelinesModule.cpp
+++ b/Engine/Plugins/Experimental/Interchange/Editor/Source/Pipelines/Private/InterchangeEditorPipelinesModule.cpp
@@ -2,17 +2,22 @@
 
 #include "InterchangeEditorPipelinesModule.h"
 
+#include "AssetToolsModule.h"
 #include "CoreMinimal.h"
 #include "Engine/Engine.h"
 #include "InterchangeEditorPipelineDetails.h"
 #include "InterchangeEditorPipelineStyle.h"
-#include "InterchangeGraphInspectorPipeline.h"
 #include "InterchangeManager.h"
+#include "InterchangePipelineBase.h"
+#include "InterchangePipelineFactories.h"
+#include "InterchangePythonPipelineBase.h"
 #include "Misc/CoreDelegates.h"
 #include "Modules/ModuleManager.h"
 #include "Nodes/InterchangeBaseNode.h"
 #include "PropertyEditorModule.h"
 
+
+
 #define LOCTEXT_NAMESPACE "InterchangeEditorPipelines"
 
 class FInterchangeEditorPipelinesModule : public IInterchangeEditorPipelinesModule
@@ -21,6 +26,12 @@ class FInterchangeEditorPipelinesModule : public IInterchangeEditorPipelinesModu
 	virtual void StartupModule() override;
 	virtual void ShutdownModule() override;
 
+	//Called when we start the module
+	void AcquireResources();
+	
+	//Can be called multiple time
+	void ReleaseResources();
+
 	TSharedRef<FPropertySection> RegisterPropertySection(FPropertyEditorModule& PropertyModule, FName ClassName, FName SectionName, FText DisplayName);
 	void RegisterPropertySectionMappings();
 	void UnregisterPropertySectionMappings();
@@ -35,6 +46,16 @@ private:
 IMPLEMENT_MODULE(FInterchangeEditorPipelinesModule, InterchangeEditorPipelines)
 
 void FInterchangeEditorPipelinesModule::StartupModule()
+{
+	AcquireResources();
+}
+
+void FInterchangeEditorPipelinesModule::ShutdownModule()
+{
+	ReleaseResources();
+}
+
+void FInterchangeEditorPipelinesModule::AcquireResources()
 {
 	auto RegisterItems = [this]()
 	{
@@ -59,13 +80,36 @@ void FInterchangeEditorPipelinesModule::StartupModule()
 	ClassesToUnregisterOnShutdown.Add(UInterchangeBaseNode::StaticClass()->GetFName());
 	PropertyEditorModule.RegisterCustomClassLayout(ClassesToUnregisterOnShutdown.Last(), FOnGetDetailCustomizationInstance::CreateStatic(&FInterchangeBaseNodeDetailsCustomization::MakeInstance));
 	
+	ClassesToUnregisterOnShutdown.Add(UInterchangePipelineBase::StaticClass()->GetFName());
+	PropertyEditorModule.RegisterCustomClassLayout(ClassesToUnregisterOnShutdown.Last(), FOnGetDetailCustomizationInstance::CreateStatic(&FInterchangePipelineBaseDetailsCustomization::MakeInstance));
+
+
 	if (!InterchangeEditorPipelineStyle.IsValid())
 	{
 		InterchangeEditorPipelineStyle = MakeShared<FInterchangeEditorPipelineStyle>();
 	}
+
+	// Register the InterchangeImportTestPlan asset
+	FAssetToolsModule& AssetToolsModule = FAssetToolsModule::GetModule();
+	IAssetTools& AssetTools = AssetToolsModule.Get();
+	
+	BlueprintPipelineBase_TypeActions = MakeShared<FAssetTypeActions_InterchangeBlueprintPipelineBase>();
+	AssetTools.RegisterAssetTypeActions(BlueprintPipelineBase_TypeActions.ToSharedRef());
+
+	PipelineBase_TypeActions = MakeShared<FAssetTypeActions_InterchangePipelineBase>();
+	AssetTools.RegisterAssetTypeActions(PipelineBase_TypeActions.ToSharedRef());
+
+	PythonPipelineBase_TypeActions = MakeShared<FAssetTypeActions_InterchangePythonPipelineBase>();
+	AssetTools.RegisterAssetTypeActions(PythonPipelineBase_TypeActions.ToSharedRef());
+
+	FCoreDelegates::OnPreExit.AddLambda([this]()
+		{
+			//We must release the resources before the application start to unload modules
+			ReleaseResources();
+		});
 }
 
-void FInterchangeEditorPipelinesModule::ShutdownModule()
+void FInterchangeEditorPipelinesModule::ReleaseResources()
 {
 	FPropertyEditorModule* PropertyEditorModule = FModuleManager::GetModulePtr<FPropertyEditorModule>("PropertyEditor");
 	if (PropertyEditorModule)
@@ -75,10 +119,20 @@ void FInterchangeEditorPipelinesModule::ShutdownModule()
 			PropertyEditorModule->UnregisterCustomClassLayout(ClassName);
 		}
 	}
+	ClassesToUnregisterOnShutdown.Empty();
 
 	UnregisterPropertySectionMappings();
 
 	InterchangeEditorPipelineStyle = nullptr;
+
+	FAssetToolsModule* AssetToolsModule = FModuleManager::GetModulePtr<FAssetToolsModule>("AssetTools");
+	if (AssetToolsModule)
+	{
+		IAssetTools& AssetTools = AssetToolsModule->Get();
+		AssetTools.UnregisterAssetTypeActions(BlueprintPipelineBase_TypeActions.ToSharedRef());
+		AssetTools.UnregisterAssetTypeActions(PipelineBase_TypeActions.ToSharedRef());
+		AssetTools.UnregisterAssetTypeActions(PythonPipelineBase_TypeActions.ToSharedRef());
+	}
 }
 
 TSharedRef<FPropertySection> FInterchangeEditorPipelinesModule::RegisterPropertySection(FPropertyEditorModule& PropertyModule, FName ClassName, FName SectionName, FText DisplayName)
@@ -150,6 +204,8 @@ void FInterchangeEditorPipelinesModule::UnregisterPropertySectionMappings()
 		PropertyModule->RemoveSection(PropertySectionIterator->Key, PropertySectionIterator->Value);
 		PropertySectionIterator.RemoveCurrent();
 	}
+
+	RegisteredPropertySections.Empty();
 }
 
 #undef LOCTEXT_NAMESPACE
diff --git a/Engine/Plugins/Experimental/Interchange/Editor/Source/Pipelines/Private/InterchangePipelineFactories.cpp b/Engine/Plugins/Experimental/Interchange/Editor/Source/Pipelines/Private/InterchangePipelineFactories.cpp
new file mode 100644
--- /dev/null
+++ b/Engine/Plugins/Experimental/Interchange/Editor/Source/Pipelines/Private/InterchangePipelineFactories.cpp
@@ -0,0 +1,223 @@
+// Copyright Epic Games, Inc. All Rights Reserved. 
+
+#include "InterchangePipelineFactories.h"
+
+#include "AssetTypeCategories.h"
+#include "ClassViewerModule.h"
+#include "InterchangeBlueprintPipelineBase.h"
+#include "InterchangePipelineBase.h"
+#include "InterchangePythonPipelineBase.h"
+#include "Kismet2/KismetEditorUtilities.h"
+#include "Kismet2/SClassPickerDialog.h"
+#include "Misc/MessageDialog.h"
+#include "UObject/Object.h"
+#include "UObject/ObjectMacros.h"
+#include "UObject/MetaData.h"
+
+/**
+ * UInterchangeBlueprintPipelineBaseFactory implementation
+ */
+
+UInterchangeBlueprintPipelineBaseFactory::UInterchangeBlueprintPipelineBaseFactory()
+{
+	bCreateNew = true;
+	bEditAfterNew = false;
+	SupportedClass = UInterchangeBlueprintPipelineBase::StaticClass();
+	ParentClass = UInterchangePipelineBase::StaticClass();
+}
+
+UObject* UInterchangeBlueprintPipelineBaseFactory::FactoryCreateNew(UClass* Class, UObject* InParent, FName Name, EObjectFlags Flags, UObject* Context, FFeedbackContext* Warn, FName CallingContext)
+{
+	// Make sure we are trying to factory a gameplay ability blueprint, then create and init one
+	check(Class->IsChildOf(UInterchangeBlueprintPipelineBase::StaticClass()));
+
+	// If they selected an interface, force the parent class to be UInterface
+	if (BlueprintType == BPTYPE_Interface)
+	{
+		ParentClass = UInterface::StaticClass();
+	}
+
+	if ((ParentClass == nullptr) || !FKismetEditorUtilities::CanCreateBlueprintOfClass(ParentClass) || !ParentClass->IsChildOf(UInterchangePipelineBase::StaticClass()))
+	{
+		FFormatNamedArguments Args;
+		Args.Add(TEXT("ClassName"), (ParentClass != nullptr) ? FText::FromString(ParentClass->GetName()) : NSLOCTEXT("UInterchangeBlueprintPipelineBaseFactory", "Null", "(null)"));
+		FMessageDialog::Open(EAppMsgType::Ok, FText::Format(NSLOCTEXT("UInterchangeBlueprintPipelineBaseFactory", "CannotCreateInterchangeBlueprintPipelineBase", "Cannot create a Interchange Blueprint Pipeline based on the class '{ClassName}'."), Args));
+		return nullptr;
+	}
+	else
+	{
+		return CastChecked<UInterchangeBlueprintPipelineBase>(FKismetEditorUtilities::CreateBlueprint(ParentClass, InParent, Name, BlueprintType, UInterchangeBlueprintPipelineBase::StaticClass(), UBlueprintGeneratedClass::StaticClass(), CallingContext));
+	}
+}
+
+UObject* UInterchangeBlueprintPipelineBaseFactory::FactoryCreateNew(UClass* Class, UObject* InParent, FName Name, EObjectFlags Flags, UObject* Context, FFeedbackContext* Warn)
+{
+	return FactoryCreateNew(Class, InParent, Name, Flags, Context, Warn, NAME_None);
+}
+
+UFactory* FAssetTypeActions_InterchangeBlueprintPipelineBase::GetFactoryForBlueprintType(UBlueprint* InBlueprint) const
+{
+	UInterchangeBlueprintPipelineBaseFactory* InterchangeBlueprintPipelineBaseFactory = NewObject<UInterchangeBlueprintPipelineBaseFactory>();
+	InterchangeBlueprintPipelineBaseFactory->ParentClass = TSubclassOf<UInterchangePipelineBase>(*InBlueprint->GeneratedClass);
+	return InterchangeBlueprintPipelineBaseFactory;
+}
+
+UClass* FAssetTypeActions_InterchangeBlueprintPipelineBase::GetSupportedClass() const
+{
+	return UInterchangeBlueprintPipelineBase::StaticClass();
+}
+
+/**
+ * UInterchangePipelineBaseFactory implementation
+ */
+
+UInterchangePipelineBaseFactory::UInterchangePipelineBaseFactory()
+{
+	bCreateNew = true;
+	bEditAfterNew = false;
+	bEditorImport = false;
+	SupportedClass = UInterchangePipelineBase::StaticClass();
+}
+
+UObject* UInterchangePipelineBaseFactory::FactoryCreateNew(UClass* Class, UObject* InParent, FName InName, EObjectFlags InFlags, UObject* Context, FFeedbackContext* Warn)
+{
+	UInterchangePipelineBase* Pipeline = nullptr;
+	if (PipelineClass && PipelineClass->IsChildOf(SupportedClass))
+	{
+		Pipeline = NewObject<UInterchangePipelineBase>(InParent, PipelineClass, InName, InFlags | RF_Transactional);
+	}
+	return Pipeline;
+}
+
+uint32 UInterchangePipelineBaseFactory::GetMenuCategories() const
+{
+	return EAssetTypeCategories::Misc;
+}
+
+FText UInterchangePipelineBaseFactory::GetDisplayName() const
+{
+	return NSLOCTEXT("UInterchangePipelineBaseFactory", "MenuEntry", "Interchange Pipeline");
+}
+
+bool UInterchangePipelineBaseFactory::ConfigureProperties()
+{
+	FClassViewerModule& ClassViewerModule = FModuleManager::LoadModuleChecked<FClassViewerModule>("ClassViewer");
+	FClassViewerInitializationOptions Options;
+	Options.Mode = EClassViewerMode::ClassPicker;
+
+	TSharedPtr<FInterchangePipelineBaseFilterViewer> Filter = MakeShared<FInterchangePipelineBaseFilterViewer>();
+	Options.ClassFilters.Add(Filter.ToSharedRef());
+
+	//CLASS_HideDropDown prevent sub pipeline to show in the list
+	Filter->DisallowedClassFlags = CLASS_Abstract | CLASS_Deprecated | CLASS_HideDropDown | CLASS_Hidden | CLASS_Transient;
+	Filter->AllowedChildrenOfClasses.Add(UInterchangePipelineBase::StaticClass());
+	//Blueprint pipeline have there own factory
+	Filter->DisallowedChildrenOfClasses.Add(UInterchangeBlueprintPipelineBase::StaticClass());
+
+	const FText TitleText = NSLOCTEXT("UInterchangePipelineBaseFactory", "CreateOptions", "Pick a Pipeline Class");
+	UClass* ChosenClass = nullptr;
+	const bool bPressedOk = SClassPickerDialog::PickClass(TitleText, Options, ChosenClass, UInterchangePipelineBase::StaticClass());
+
+	if (bPressedOk)
+	{
+		PipelineClass = ChosenClass;
+	}
+
+	return bPressedOk;
+}
+
+UClass* FAssetTypeActions_InterchangePipelineBase::GetSupportedClass() const
+{
+	return UInterchangePipelineBase::StaticClass();
+}
+
+void FAssetTypeActions_InterchangePipelineBase::OpenAssetEditor(const TArray<UObject*>& InObjects, TSharedPtr<class IToolkitHost> EditWithinLevelEditor)
+{
+	for (UObject* Object : InObjects)
+	{
+		if (UInterchangePipelineBase* Asset = Cast<UInterchangePipelineBase>(Object))
+		{
+			FSimpleAssetEditor::CreateEditor(EToolkitMode::Standalone, EditWithinLevelEditor, Object);
+		}
+	}
+}
+
+/**
+ * UInterchangePythonPipelineFactory implementation
+ */
+
+UInterchangePythonPipelineAssetFactory::UInterchangePythonPipelineAssetFactory()
+{
+	bCreateNew = true;
+	bEditAfterNew = false;
+	bEditorImport = false;
+	SupportedClass = UInterchangePythonPipelineAsset::StaticClass();
+}
+
+UObject* UInterchangePythonPipelineAssetFactory::FactoryCreateNew(UClass* Class, UObject* InParent, FName InName, EObjectFlags InFlags, UObject* Context, FFeedbackContext* Warn)
+{
+	UInterchangePythonPipelineAsset* Pipeline = nullptr;
+	if (PythonClass && PythonClass->IsChildOf(UInterchangePipelineBase::StaticClass()))
+	{
+		Pipeline = NewObject<UInterchangePythonPipelineAsset>(InParent, SupportedClass, InName, InFlags | RF_Transactional);
+		Pipeline->PythonClass = PythonClass;
+		Pipeline->GeneratePipeline();
+	}
+	return Pipeline;
+}
+
+uint32 UInterchangePythonPipelineAssetFactory::GetMenuCategories() const
+{
+	return EAssetTypeCategories::Misc;
+}
+
+FText UInterchangePythonPipelineAssetFactory::GetDisplayName() const
+{
+	return NSLOCTEXT("UInterchangePythonPipelineFactory", "MenuEntry", "Interchange Python Pipeline");
+}
+
+bool UInterchangePythonPipelineAssetFactory::ConfigureProperties()
+{
+
+	FClassViewerModule& ClassViewerModule = FModuleManager::LoadModuleChecked<FClassViewerModule>("ClassViewer");
+	FClassViewerInitializationOptions Options;
+	Options.Mode = EClassViewerMode::ClassPicker;
+
+	TSharedPtr<FInterchangePipelineBaseFilterViewer> Filter = MakeShared<FInterchangePipelineBaseFilterViewer>();
+	Options.ClassFilters.Add(Filter.ToSharedRef());
+
+	//CLASS_HideDropDown prevent sub pipeline to show in the list
+	Filter->DisallowedClassFlags = CLASS_Abstract | CLASS_Deprecated | CLASS_HideDropDown | CLASS_Hidden | CLASS_Transient;
+
+	Filter->AllowedChildrenOfClasses.Add(UInterchangePythonPipelineBase::StaticClass());
+
+	//Blueprint pipeline have there own factory
+	Filter->DisallowedChildrenOfClasses.Add(UInterchangeBlueprintPipelineBase::StaticClass());
+
+	const FText TitleText = NSLOCTEXT("UInterchangePythonPipelineAssetFactory", "CreateOptions", "Pick a Pipeline Class");
+	UClass* ChosenClass = nullptr;
+	const bool bPressedOk = SClassPickerDialog::PickClass(TitleText, Options, ChosenClass, UInterchangePipelineBase::StaticClass());
+
+	if (bPressedOk)
+	{
+		PythonClass = ChosenClass;
+	}
+
+	return bPressedOk;
+}
+
+UClass* FAssetTypeActions_InterchangePythonPipelineBase::GetSupportedClass() const
+{
+	return UInterchangePythonPipelineAsset::StaticClass();
+}
+
+void FAssetTypeActions_InterchangePythonPipelineBase::OpenAssetEditor(const TArray<UObject*>& InObjects, TSharedPtr<class IToolkitHost> EditWithinLevelEditor)
+{
+	for (UObject* Object : InObjects)
+	{
+		if (UInterchangePythonPipelineAsset* Asset = Cast<UInterchangePythonPipelineAsset>(Object))
+		{
+			FSimpleAssetEditor::CreateEditor(EToolkitMode::Standalone, EditWithinLevelEditor, Object);
+		}
+	}
+}
\ No newline at end of file
diff --git a/Engine/Plugins/Experimental/Interchange/Editor/Source/Pipelines/Private/InterchangePipelineFactories.h b/Engine/Plugins/Experimental/Interchange/Editor/Source/Pipelines/Private/InterchangePipelineFactories.h
new file mode 100644
--- /dev/null
+++ b/Engine/Plugins/Experimental/Interchange/Editor/Source/Pipelines/Private/InterchangePipelineFactories.h
@@ -0,0 +1,136 @@
+// Copyright Epic Games, Inc. All Rights Reserved.
+
+#pragma once
+
+#include "AssetTypeActions/AssetTypeActions_Blueprint.h"
+#include "AssetTypeActions_Base.h"
+#include "ClassViewerFilter.h"
+#include "Factories/Factory.h"
+#include "UObject/Object.h"
+#include "UObject/ObjectMacros.h"
+
+#include "InterchangePipelineFactories.generated.h"
+
+class FInterchangePipelineBaseFilterViewer : public IClassViewerFilter
+{
+public:
+	TSet<const UClass*> AllowedChildrenOfClasses;
+	TSet<const UClass*> DisallowedChildrenOfClasses;
+	EClassFlags DisallowedClassFlags = EClassFlags::CLASS_None;
+	virtual bool IsClassAllowed(const FClassViewerInitializationOptions& InInitOptions, const UClass* InClass, TSharedRef<class FClassViewerFilterFuncs> InFilterFuncs) override
+	{
+		if (InClass->HasAnyClassFlags(DisallowedClassFlags))
+		{
+			return false;
+		}
+		
+		return InFilterFuncs->IfInChildOfClassesSet(AllowedChildrenOfClasses, InClass) != EFilterReturn::Failed
+			&& InFilterFuncs->IfInChildOfClassesSet(DisallowedChildrenOfClasses, InClass) == EFilterReturn::Failed;
+	}
+
+	virtual bool IsUnloadedClassAllowed(const FClassViewerInitializationOptions& InInitOptions, const TSharedRef<const IUnloadedBlueprintData> InUnloadedClassData, TSharedRef<FClassViewerFilterFuncs> InFilterFuncs) override
+	{
+		if (InUnloadedClassData->HasAnyClassFlags(DisallowedClassFlags))
+		{
+			return false;
+		}
+
+		return InFilterFuncs->IfInChildOfClassesSet(AllowedChildrenOfClasses, InUnloadedClassData) != EFilterReturn::Failed
+			&& InFilterFuncs->IfInChildOfClassesSet(DisallowedChildrenOfClasses, InUnloadedClassData) == EFilterReturn::Failed;
+	}
+};
+
+UCLASS(hidecategories = Object, collapsecategories)
+class UInterchangeBlueprintPipelineBaseFactory : public UFactory
+{
+	GENERATED_BODY()
+
+public:
+	UInterchangeBlueprintPipelineBaseFactory();
+
+	// The type of blueprint that will be created
+	UPROPERTY(EditAnywhere, Category = InterchangeBlueprintPipelineBaseFactory)
+	TEnumAsByte<enum EBlueprintType> BlueprintType;
+
+	// The parent class of the created blueprint
+	UPROPERTY(EditAnywhere, Category = InterchangeBlueprintPipelineBaseFactory)
+	TSubclassOf<class UInterchangePipelineBase> ParentClass;
+
+	//Begin UFactory Interface
+	virtual UObject* FactoryCreateNew(UClass* Class, UObject* InParent, FName Name, EObjectFlags Flags, UObject* Context, FFeedbackContext* Warn, FName CallingContext) override;
+	virtual UObject* FactoryCreateNew(UClass* Class, UObject* InParent, FName Name, EObjectFlags Flags, UObject* Context, FFeedbackContext* Warn) override;
+	//End UFactory Interface	
+};
+
+class FAssetTypeActions_InterchangeBlueprintPipelineBase : public FAssetTypeActions_Blueprint
+{
+public:
+	// FAssetTypeActions_Blueprint interface
+	virtual UFactory* GetFactoryForBlueprintType(UBlueprint* InBlueprint) const override;
+
+	// IAssetTypeActions Implementation
+	virtual FText GetName() const override { return NSLOCTEXT("FAssetTypeActions_InterchangeGraphInspectorPipeline", "InterchangeGraphInspectorPipelineName", "Interchange Blueprint Pipeline Base"); }
+	virtual FColor GetTypeColor() const override { return FColor(10, 25, 175); }
+	virtual UClass* GetSupportedClass() const override;
+	virtual uint32 GetCategories() override { return EAssetTypeCategories::Blueprint | EAssetTypeCategories::Misc; }
+};
+
+UCLASS(hidecategories = Object, collapsecategories)
+class UInterchangePipelineBaseFactory : public UFactory
+{
+	GENERATED_BODY()
+
+public:
+	UInterchangePipelineBaseFactory();
+
+	// UFactory Interface
+	virtual UObject* FactoryCreateNew(UClass* Class, UObject* InParent, FName Name, EObjectFlags Flags, UObject* Context, FFeedbackContext* Warn) override;
+	virtual uint32 GetMenuCategories() const override;
+	virtual FText GetDisplayName() const override;
+	virtual bool ConfigureProperties() override;
+	// End of UFactory Interface
+private:
+	UClass* PipelineClass = nullptr;
+};
+
+
+
+class FAssetTypeActions_InterchangePipelineBase : public FAssetTypeActions_Base
+{
+public:
+	// IAssetTypeActions Implementation
+	virtual FText GetName() const override { return NSLOCTEXT("FAssetTypeActions_InterchangePipelineBase", "InterchangePipelineBaseName", "Interchange Pipeline"); }
+	virtual FColor GetTypeColor() const override { return FColor(135, 200, 25); }
+	virtual UClass* GetSupportedClass() const override;
+	virtual uint32 GetCategories() override { return EAssetTypeCategories::Misc; }
+	virtual void OpenAssetEditor(const TArray<UObject*>& InObjects, TSharedPtr<class IToolkitHost> EditWithinLevelEditor = TSharedPtr<IToolkitHost>()) override;
+};
+
+UCLASS(hidecategories = Object, collapsecategories)
+class UInterchangePythonPipelineAssetFactory : public UFactory
+{
+	GENERATED_BODY()
+
+public:
+	UInterchangePythonPipelineAssetFactory();
+
+	// UFactory Interface
+	virtual UObject* FactoryCreateNew(UClass* Class, UObject* InParent, FName Name, EObjectFlags Flags, UObject* Context, FFeedbackContext* Warn) override;
+	virtual uint32 GetMenuCategories() const override;
+	virtual FText GetDisplayName() const override;
+	virtual bool ConfigureProperties() override;
+	// End of UFactory Interface
+private:
+	UClass* PythonClass = nullptr;
+};
+
+class FAssetTypeActions_InterchangePythonPipelineBase : public FAssetTypeActions_Base
+{
+public:
+	// IAssetTypeActions Implementation
+	virtual FText GetName() const override { return NSLOCTEXT("InterchangePipelineFactories", "FAssetTypeActions_InterchangePythonPipeline", "Interchange Python Pipeline"); }
+	virtual FColor GetTypeColor() const override { return FColor(135, 200, 25); }
+	virtual UClass* GetSupportedClass() const override;
+	virtual uint32 GetCategories() override { return EAssetTypeCategories::Misc; }
+	virtual void OpenAssetEditor(const TArray<UObject*>& InObjects, TSharedPtr<class IToolkitHost> EditWithinLevelEditor = TSharedPtr<IToolkitHost>()) override;
+};
\ No newline at end of file
diff --git a/Engine/Plugins/Experimental/Interchange/Editor/Source/Pipelines/Private/SInterchangePipelineConfigurationDialog.cpp b/Engine/Plugins/Experimental/Interchange/Editor/Source/Pipelines/Private/SInterchangePipelineConfigurationDialog.cpp
--- a/Engine/Plugins/Experimental/Interchange/Editor/Source/Pipelines/Private/SInterchangePipelineConfigurationDialog.cpp
+++ b/Engine/Plugins/Experimental/Interchange/Editor/Source/Pipelines/Private/SInterchangePipelineConfigurationDialog.cpp
@@ -6,11 +6,15 @@
 #include "GameFramework/Actor.h"
 #include "IDetailsView.h"
 #include "IDocumentation.h"
+#include "InterchangeBlueprintPipelineBase.h"
+#include "InterchangeManager.h"
 #include "InterchangePipelineBase.h"
 #include "InterchangeProjectSettings.h"
+#include "InterchangePythonPipelineBase.h"
 #include "InterchangeSourceData.h"
 #include "Misc/App.h"
 #include "Modules/ModuleManager.h"
+#include "PropertyEditorDelegates.h"
 #include "PropertyEditorModule.h"
 #include "Styling/SlateIconFinder.h"
 #include "Widgets/Images/SImage.h"
@@ -65,10 +69,12 @@ void SInterchangePipelineStacksTreeView::Construct(const FArguments& InArgs)
 		{
 			if(UInterchangePipelineBase* GeneratedPipeline = PipelineStack[PipelineIndex])
 			{
+				//When we re-import we save the setting so ResetToDefault can simply just reload the settings
 				GeneratedPipeline->SaveSettings(ReimportPipelineName);
-				//Load the settings for this pipeline
+				GeneratedPipeline->bAllowLockedPropertiesEdition = false;
 				TSharedPtr<FInterchangePipelineStacksTreeNodeItem> PipelineNode = MakeShared<FInterchangePipelineStacksTreeNodeItem>();
 				PipelineNode->StackName = ReimportPipelineName;
+				//Currently the re-import never lock properties
 				PipelineNode->Pipeline = GeneratedPipeline;
 				StackNode->Childrens.Add(PipelineNode);
 			}
@@ -85,10 +91,9 @@ void SInterchangePipelineStacksTreeView::Construct(const FArguments& InArgs)
 			const FInterchangePipelineStack& InterchangePipelineStack = NameAndPipelineStack.Value;
 			for (int32 PipelineIndex = 0; PipelineIndex < InterchangePipelineStack.Pipelines.Num(); ++PipelineIndex)
 			{
-				const UClass* PipelineClass = InterchangePipelineStack.Pipelines[PipelineIndex].LoadSynchronous();
-				if (PipelineClass)
+				if (UInterchangePipelineBase* GeneratedPipeline = UE::Interchange::GeneratePipelineInstance(InterchangePipelineStack.Pipelines[PipelineIndex]))
 				{
-					UInterchangePipelineBase* GeneratedPipeline = NewObject<UInterchangePipelineBase>(GetTransientPackage(), PipelineClass, NAME_None, RF_NoFlags);
+					GeneratedPipeline->bAllowLockedPropertiesEdition = false;
 					//Load the settings for this pipeline
 					GeneratedPipeline->LoadSettings(NameAndPipelineStack.Key);
 					GeneratedPipeline->PreDialogCleanup(NameAndPipelineStack.Key);
@@ -106,7 +111,7 @@ void SInterchangePipelineStacksTreeView::Construct(const FArguments& InArgs)
 	(
 		STreeView::FArguments()
 		.TreeItemsSource(&RootNodeArray)
-		.SelectionMode(ESelectionMode::Multi)
+		.SelectionMode(ESelectionMode::Single)
 		.OnGenerateRow(this, &SInterchangePipelineStacksTreeView::OnGenerateRowPipelineConfigurationTreeView)
 		.OnGetChildren(this, &SInterchangePipelineStacksTreeView::OnGetChildrenPipelineConfigurationTreeView)
 		.OnContextMenuOpening(this, &SInterchangePipelineStacksTreeView::OnOpenContextMenu)
@@ -398,6 +403,10 @@ TSharedRef<SBox> SInterchangePipelineConfigurationDialog::SpawnPipelineConfigura
 	PipelineConfigurationDetailsView = PropertyEditorModule.CreateDetailView(DetailsViewArgs);
 	InspectorBox->SetContent(PipelineConfigurationDetailsView->AsShared());
 	PipelineConfigurationDetailsView->SetObject(nullptr);
+	PipelineConfigurationDetailsView->GetIsPropertyVisibleDelegate().BindLambda([this](const FPropertyAndParent& PropertyAndParent)
+		{
+			return IsPropertyVisible(PropertyAndParent);
+		});
 	return PipelineConfigurationPanelBox;
 }
 
@@ -422,11 +431,27 @@ void SInterchangePipelineConfigurationDialog::Construct(const FArguments& InArgs
 			SNew(SVerticalBox)
 			+ SVerticalBox::Slot()
 			.AutoHeight()
-			.HAlign(HAlign_Left)
+			.HAlign(HAlign_Fill)
 			.Padding(2)
 			[
-				SNew(STextBlock)
-				.Text(this, &SInterchangePipelineConfigurationDialog::GetSourceDescription)
+				SNew(SHorizontalBox)
+				+SHorizontalBox::Slot()
+				.FillWidth(1.0f)
+				.Padding(2)
+				[
+					SNew(STextBlock)
+					.Text(this, &SInterchangePipelineConfigurationDialog::GetSourceDescription)
+				]
+				+SHorizontalBox::Slot()
+				.AutoWidth()
+				.HAlign(HAlign_Right)
+				.Padding(10, 2, 2, 2)
+				[
+					SNew(SButton)
+					.HAlign(HAlign_Center)
+					.Text(LOCTEXT("InspectorGraphWindow_ResetDefault", "Reset to Default"))
+					.OnClicked(this, &SInterchangePipelineConfigurationDialog::OnResetToDefault)
+				]
 			]
 			+ SVerticalBox::Slot()
 			.FillHeight(1.0f)
@@ -477,9 +502,25 @@ void SInterchangePipelineConfigurationDialog::Construct(const FArguments& InArgs
 
 void SInterchangePipelineConfigurationDialog::OnSelectionChanged(TSharedPtr<FInterchangePipelineStacksTreeNodeItem> Item, ESelectInfo::Type SelectionType)
 {
-	UInterchangePipelineBase* Pipeline = Item ? Item->Pipeline : nullptr;
-	//Change the object point by the InspectorBox
-	PipelineConfigurationDetailsView->SetObject(Pipeline);
+	CurrentSelectedPipeline = nullptr;
+	FName StackName = NAME_None;
+	if (Item.IsValid())
+	{
+		CurrentSelectedPipeline = Item->Pipeline;
+		StackName = Item->StackName;
+	}
+	PipelineConfigurationDetailsView->SetObject(CurrentSelectedPipeline.Get());
+	CurrentStackName = StackName;
+}
+
+bool SInterchangePipelineConfigurationDialog::IsPropertyVisible(const FPropertyAndParent& PropertyAndParent) const
+{
+	if (bReimport)
+	{
+		const FName ReimportRestrictKey(TEXT("ReimportRestrict"));
+		return !(PropertyAndParent.Property.GetBoolMetaData(ReimportRestrictKey));
+	}
+	return true;
 }
 
 FText SInterchangePipelineConfigurationDialog::GetSourceDescription() const
@@ -500,6 +541,90 @@ FText SInterchangePipelineConfigurationDialog::GetSourceDescription() const
 	return ActionDescription;
 }
 
+void SInterchangePipelineConfigurationDialog::RecursiveIterateNode(TSharedPtr<FInterchangePipelineStacksTreeNodeItem>& ParentNode, TFunctionRef<void(TSharedPtr<FInterchangePipelineStacksTreeNodeItem>&)> IterationLambda)
+{
+	IterationLambda(ParentNode);
+
+	for (int32 ChildIndex = 0; ChildIndex < ParentNode->Childrens.Num(); ++ChildIndex)
+	{
+		RecursiveIterateNode(ParentNode->Childrens[ChildIndex], IterationLambda);
+	}
+}
+
+FReply SInterchangePipelineConfigurationDialog::OnResetToDefault()
+{
+	FReply Result = FReply::Handled();
+	TArray<TWeakObjectPtr<UObject>> SelectedPipelines = PipelineConfigurationDetailsView->GetSelectedObjects();
+	if (CurrentStackName == NAME_None)
+	{
+		return Result;
+	}
+	//Multi selection is not allowed
+	ensure(SelectedPipelines.Num() <= 1);
+	for(TWeakObjectPtr<UObject> WeakObject : SelectedPipelines)
+	{
+		if (UInterchangePipelineBase* Pipeline = Cast<UInterchangePipelineBase>(WeakObject.Get()))
+		{
+			if (bReimport)
+			{
+				//When we re-import we save the settings before showing the dialog, so the reset to default simply reload
+				//the original settings when the dialog was loaded.
+				for (int32 PipelineIndex = 0; PipelineIndex < PipelineStack.Num(); ++PipelineIndex)
+				{
+					if (Pipeline == PipelineStack[PipelineIndex])
+					{
+						Pipeline->LoadSettings(CurrentStackName);
+						PipelineConfigurationDetailsView->SetObject(Pipeline, true);
+						break;
+					}
+				}
+			}
+			else
+			{
+				//The reset to default when doing a import is to duplicate again the pipeline stack reference asset.
+				const TMap<FName, FInterchangePipelineStack>& DefaultPipelineStacks = GetDefault<UInterchangeProjectSettings>()->PipelineStacks;
+				//Find the stack node, so we can update the pipeline
+				TSharedPtr<FInterchangePipelineStacksTreeNodeItem> StackNode;
+				TArray<TSharedPtr<FInterchangePipelineStacksTreeNodeItem>>& RootNodeArray = PipelineConfigurationTreeView->GetMutableRootNodeArray();
+				for (TSharedPtr<FInterchangePipelineStacksTreeNodeItem>& RootNode : RootNodeArray)
+				{
+					RecursiveIterateNode(RootNode, [this, &Pipeline, &DefaultPipelineStacks](TSharedPtr<FInterchangePipelineStacksTreeNodeItem>& NodeItem)
+						{
+							if (NodeItem->StackName == CurrentStackName && NodeItem->Pipeline == Pipeline)
+							{
+								//We assume the pipelines inside one stack are all different classes, we use the class to know which default asset we need to duplicate
+								const UClass* PipelineClass = NodeItem->Pipeline->GetClass();
+								for (const TPair<FName, FInterchangePipelineStack>& NameAndPipelineStack : DefaultPipelineStacks)
+								{
+									if(CurrentStackName != NameAndPipelineStack.Key)
+									{
+										continue;
+									}
+									const FInterchangePipelineStack& InterchangePipelineStack = NameAndPipelineStack.Value;
+									for (int32 PipelineIndex = 0; PipelineIndex < InterchangePipelineStack.Pipelines.Num(); ++PipelineIndex)
+									{
+										if (UInterchangePipelineBase* GeneratedPipeline = UE::Interchange::GeneratePipelineInstance(InterchangePipelineStack.Pipelines[PipelineIndex]))
+										{
+											if (GeneratedPipeline->GetClass() == PipelineClass)
+											{
+												NodeItem->Pipeline = GeneratedPipeline;
+												NodeItem->Pipeline->bAllowLockedPropertiesEdition = false;
+												PipelineConfigurationDetailsView->SetObject(NodeItem->Pipeline, true);
+												//Exit the lambda
+												return;
+											}
+										}
+									}
+								}
+							}
+						});
+				}
+			}
+		}
+	}
+	return Result;
+}
+
 bool SInterchangePipelineConfigurationDialog::RecursiveValidatePipelineSettings(const TSharedPtr<FInterchangePipelineStacksTreeNodeItem>& ParentNode, TOptional<FText>& OutInvalidReason) const
 {
 	if (ParentNode->Pipeline)
diff --git a/Engine/Plugins/Experimental/Interchange/Editor/Source/Pipelines/Private/SInterchangePipelineConfigurationDialog.h b/Engine/Plugins/Experimental/Interchange/Editor/Source/Pipelines/Private/SInterchangePipelineConfigurationDialog.h
--- a/Engine/Plugins/Experimental/Interchange/Editor/Source/Pipelines/Private/SInterchangePipelineConfigurationDialog.h
+++ b/Engine/Plugins/Experimental/Interchange/Editor/Source/Pipelines/Private/SInterchangePipelineConfigurationDialog.h
@@ -2,18 +2,19 @@
 
 #pragma once
 
-#include "CoreMinimal.h"
 #include "InputCoreTypes.h"
 #include "Input/Reply.h"
 #include "InterchangePipelineBase.h"
 #include "InterchangeSourceData.h"
 #include "Styling/SlateColor.h"
 #include "UObject/GCObject.h"
+#include "UObject/WeakObjectPtrTemplates.h"
 #include "Widgets/DeclarativeSyntaxSupport.h"
 #include "Widgets/SCompoundWidget.h"
 #include "Widgets/Views/STreeView.h"
 
 class IDetailsView;
+struct FPropertyAndParent;
 
 class FInterchangePipelineStacksTreeNodeItem : protected FGCObject
 {
@@ -64,6 +65,7 @@ public:
 	FReply OnCollapseAll();
 
 	const TArray<TSharedPtr<FInterchangePipelineStacksTreeNodeItem>>& GetRootNodeArray() const { return RootNodeArray; }
+	TArray<TSharedPtr<FInterchangePipelineStacksTreeNodeItem>>& GetMutableRootNodeArray() { return RootNodeArray; }
 protected:
 	/** Delegate to invoke when selection changes. */
 	FOnPipelineConfigurationSelectionChanged OnSelectionChangedDelegate;
@@ -130,7 +132,10 @@ private:
 	TSharedRef<SBox> SpawnPipelineConfiguration();
 	void OnSelectionChanged(TSharedPtr<FInterchangePipelineStacksTreeNodeItem> Item, ESelectInfo::Type SelectionType);
 
+	bool IsPropertyVisible(const FPropertyAndParent&) const;
 	FText GetSourceDescription() const;
+	void RecursiveIterateNode(TSharedPtr<FInterchangePipelineStacksTreeNodeItem>& ParentNode, TFunctionRef<void(TSharedPtr<FInterchangePipelineStacksTreeNodeItem>& CurrentNode)> IterationLambda);
+	FReply OnResetToDefault();
 
 	bool RecursiveValidatePipelineSettings(const TSharedPtr<FInterchangePipelineStacksTreeNodeItem>& ParentNode, TOptional<FText>& OutInvalidReason) const;
 	bool IsImportButtonEnabled() const;
@@ -145,5 +150,8 @@ private:
 
 	bool bCanceled = false;
 	bool bImportAll = false;
+
+	FName CurrentStackName = NAME_None;
+	TWeakObjectPtr<UInterchangePipelineBase> CurrentSelectedPipeline = nullptr;
 };
 
diff --git a/Engine/Plugins/Experimental/Interchange/Editor/Source/Pipelines/Public/InterchangeEditorPipelinesModule.h b/Engine/Plugins/Experimental/Interchange/Editor/Source/Pipelines/Public/InterchangeEditorPipelinesModule.h
--- a/Engine/Plugins/Experimental/Interchange/Editor/Source/Pipelines/Public/InterchangeEditorPipelinesModule.h
+++ b/Engine/Plugins/Experimental/Interchange/Editor/Source/Pipelines/Public/InterchangeEditorPipelinesModule.h
@@ -7,6 +7,9 @@
 #include "Modules/ModuleManager.h"
 #include "Styling/ISlateStyle.h"
 
+class FAssetTypeActions_InterchangeBlueprintPipelineBase;
+class FAssetTypeActions_InterchangePipelineBase;
+class FAssetTypeActions_InterchangePythonPipelineBase;
 
 /**
  * The public interface to this module
@@ -38,5 +41,9 @@ public:
 	}
 
 	TArray<FName> ClassesToUnregisterOnShutdown;
+
+	TSharedPtr<FAssetTypeActions_InterchangeBlueprintPipelineBase> BlueprintPipelineBase_TypeActions;
+	TSharedPtr<FAssetTypeActions_InterchangePipelineBase> PipelineBase_TypeActions;
+	TSharedPtr<FAssetTypeActions_InterchangePythonPipelineBase> PythonPipelineBase_TypeActions;
 };
 
diff --git a/Engine/Plugins/Experimental/Interchange/Editor/Source/Pipelines/Public/InterchangePipelineConfigurationGeneric.h b/Engine/Plugins/Experimental/Interchange/Editor/Source/Pipelines/Public/InterchangePipelineConfigurationGeneric.h
--- a/Engine/Plugins/Experimental/Interchange/Editor/Source/Pipelines/Public/InterchangePipelineConfigurationGeneric.h
+++ b/Engine/Plugins/Experimental/Interchange/Editor/Source/Pipelines/Public/InterchangePipelineConfigurationGeneric.h
@@ -11,7 +11,7 @@
 
 class UInterchangeSourceData;
 
-UCLASS(BlueprintType, Blueprintable)
+UCLASS(BlueprintType, editinlinenew, Experimental)
 class INTERCHANGEEDITORPIPELINES_API UInterchangePipelineConfigurationGeneric : public UInterchangePipelineConfigurationBase
 {
 	GENERATED_BODY()
diff --git a/Engine/Plugins/Experimental/Interchange/Runtime/Source/Import/Private/Animation/InterchangeAnimSequenceFactory.cpp b/Engine/Plugins/Experimental/Interchange/Runtime/Source/Import/Private/Animation/InterchangeAnimSequenceFactory.cpp
--- a/Engine/Plugins/Experimental/Interchange/Runtime/Source/Import/Private/Animation/InterchangeAnimSequenceFactory.cpp
+++ b/Engine/Plugins/Experimental/Interchange/Runtime/Source/Import/Private/Animation/InterchangeAnimSequenceFactory.cpp
@@ -964,7 +964,7 @@ void UInterchangeAnimSequenceFactory::PreImportPreCompletedCallback(const FImpor
 		UAnimSequence* AnimSequence = CastChecked<UAnimSequence>(Arguments.ImportedObject);
 
 		UAssetImportData* ImportDataPtr = AnimSequence->AssetImportData;
-		UE::Interchange::FFactoryCommon::FUpdateImportAssetDataParameters UpdateImportAssetDataParameters(AnimSequence, ImportDataPtr, Arguments.SourceData, Arguments.NodeUniqueID, Arguments.NodeContainer, Arguments.Pipelines);
+		UE::Interchange::FFactoryCommon::FUpdateImportAssetDataParameters UpdateImportAssetDataParameters(AnimSequence, ImportDataPtr, Arguments.SourceData, Arguments.NodeUniqueID, Arguments.NodeContainer, Arguments.OriginalPipelines);
 		ImportDataPtr = UE::Interchange::FFactoryCommon::UpdateImportAssetData(UpdateImportAssetDataParameters);
 		AnimSequence->AssetImportData = ImportDataPtr;
 	}
diff --git a/Engine/Plugins/Experimental/Interchange/Runtime/Source/Import/Private/Material/InterchangeMaterialFactory.cpp b/Engine/Plugins/Experimental/Interchange/Runtime/Source/Import/Private/Material/InterchangeMaterialFactory.cpp
--- a/Engine/Plugins/Experimental/Interchange/Runtime/Source/Import/Private/Material/InterchangeMaterialFactory.cpp
+++ b/Engine/Plugins/Experimental/Interchange/Runtime/Source/Import/Private/Material/InterchangeMaterialFactory.cpp
@@ -488,7 +488,7 @@ void UInterchangeMaterialFactory::PreImportPreCompletedCallback(const FImportPre
 																										  , Arguments.SourceData
 																										  , Arguments.NodeUniqueID
 																										  , Arguments.NodeContainer
-																										  , Arguments.Pipelines);
+																										  , Arguments.OriginalPipelines);
 
 		ImportedMaterialInterface->AssetImportData = UE::Interchange::FFactoryCommon::UpdateImportAssetData(UpdateImportAssetDataParameters);
 #endif // WITH_EDITORONLY_DATA
diff --git a/Engine/Plugins/Experimental/Interchange/Runtime/Source/Import/Private/Mesh/InterchangeSkeletalMeshFactory.cpp b/Engine/Plugins/Experimental/Interchange/Runtime/Source/Import/Private/Mesh/InterchangeSkeletalMeshFactory.cpp
--- a/Engine/Plugins/Experimental/Interchange/Runtime/Source/Import/Private/Mesh/InterchangeSkeletalMeshFactory.cpp
+++ b/Engine/Plugins/Experimental/Interchange/Runtime/Source/Import/Private/Mesh/InterchangeSkeletalMeshFactory.cpp
@@ -1448,7 +1448,7 @@ void UInterchangeSkeletalMeshFactory::PreImportPreCompletedCallback(const FImpor
 																										  , Arguments.SourceData
 																										  , Arguments.NodeUniqueID
 																										  , Arguments.NodeContainer
-																										  , Arguments.Pipelines);
+																										  , Arguments.OriginalPipelines);
 
 		ImportDataPtr = UE::Interchange::FFactoryCommon::UpdateImportAssetData(UpdateImportAssetDataParameters, [&Arguments, SkeletalMesh](UInterchangeAssetImportData* AssetImportData)
 			{
diff --git a/Engine/Plugins/Experimental/Interchange/Runtime/Source/Import/Private/Mesh/InterchangeStaticMeshFactory.cpp b/Engine/Plugins/Experimental/Interchange/Runtime/Source/Import/Private/Mesh/InterchangeStaticMeshFactory.cpp
--- a/Engine/Plugins/Experimental/Interchange/Runtime/Source/Import/Private/Mesh/InterchangeStaticMeshFactory.cpp
+++ b/Engine/Plugins/Experimental/Interchange/Runtime/Source/Import/Private/Mesh/InterchangeStaticMeshFactory.cpp
@@ -494,7 +494,7 @@ void UInterchangeStaticMeshFactory::PreImportPreCompletedCallback(const FImportP
 		UStaticMesh* StaticMesh = CastChecked<UStaticMesh>(Arguments.ImportedObject);
 
 		UAssetImportData* ImportDataPtr = StaticMesh->GetAssetImportData();
-		UE::Interchange::FFactoryCommon::FUpdateImportAssetDataParameters UpdateImportAssetDataParameters(StaticMesh, ImportDataPtr, Arguments.SourceData, Arguments.NodeUniqueID, Arguments.NodeContainer, Arguments.Pipelines);
+		UE::Interchange::FFactoryCommon::FUpdateImportAssetDataParameters UpdateImportAssetDataParameters(StaticMesh, ImportDataPtr, Arguments.SourceData, Arguments.NodeUniqueID, Arguments.NodeContainer, Arguments.OriginalPipelines);
 		ImportDataPtr = UE::Interchange::FFactoryCommon::UpdateImportAssetData(UpdateImportAssetDataParameters);
 		StaticMesh->SetAssetImportData(ImportDataPtr);
 	}
diff --git a/Engine/Plugins/Experimental/Interchange/Runtime/Source/Import/Private/Texture/InterchangeTextureFactory.cpp b/Engine/Plugins/Experimental/Interchange/Runtime/Source/Import/Private/Texture/InterchangeTextureFactory.cpp
--- a/Engine/Plugins/Experimental/Interchange/Runtime/Source/Import/Private/Texture/InterchangeTextureFactory.cpp
+++ b/Engine/Plugins/Experimental/Interchange/Runtime/Source/Import/Private/Texture/InterchangeTextureFactory.cpp
@@ -1384,7 +1384,7 @@ void UInterchangeTextureFactory::PreImportPreCompletedCallback(const FImportPreC
 																						 , Arguments.SourceData
 																						 , Arguments.NodeUniqueID
 																						 , Arguments.NodeContainer
-																						 , Arguments.Pipelines);
+																						 , Arguments.OriginalPipelines);
 		SetImportAssetDataParameters.SourceFiles = MoveTemp(SourceFiles);
 
 		Texture->AssetImportData = UE::Interchange::FFactoryCommon::SetImportAssetData(SetImportAssetDataParameters);
diff --git a/Engine/Plugins/Experimental/Interchange/Runtime/Source/Pipelines/Private/InterchangeGenericMeshPipeline.cpp b/Engine/Plugins/Experimental/Interchange/Runtime/Source/Pipelines/Private/InterchangeGenericMeshPipeline.cpp
--- a/Engine/Plugins/Experimental/Interchange/Runtime/Source/Pipelines/Private/InterchangeGenericMeshPipeline.cpp
+++ b/Engine/Plugins/Experimental/Interchange/Runtime/Source/Pipelines/Private/InterchangeGenericMeshPipeline.cpp
@@ -31,6 +31,11 @@ void UInterchangeGenericMeshPipeline::AdjustSettingsForReimportType(EInterchange
 	}
 }
 
+void UInterchangeGenericMeshPipeline::PreDialogCleanup(const FName PipelineStackName)
+{
+	PhysicsAsset = nullptr;
+}
+
 void UInterchangeGenericMeshPipeline::ExecutePreImportPipeline(UInterchangeBaseNodeContainer* InBaseNodeContainer, const TArray<UInterchangeSourceData*>& InSourceDatas)
 {
 	if (!InBaseNodeContainer)
diff --git a/Engine/Plugins/Experimental/Interchange/Runtime/Source/Pipelines/Private/InterchangeGenericSkeletalMeshPipeline.cpp b/Engine/Plugins/Experimental/Interchange/Runtime/Source/Pipelines/Private/InterchangeGenericSkeletalMeshPipeline.cpp
--- a/Engine/Plugins/Experimental/Interchange/Runtime/Source/Pipelines/Private/InterchangeGenericSkeletalMeshPipeline.cpp
+++ b/Engine/Plugins/Experimental/Interchange/Runtime/Source/Pipelines/Private/InterchangeGenericSkeletalMeshPipeline.cpp
@@ -438,7 +438,7 @@ UInterchangeSkeletalMeshFactoryNode* UInterchangeGenericMeshPipeline::CreateSkel
 		UInterchangePhysicsAssetFactoryNode* PhysicsAssetFactoryNode = NewObject<UInterchangePhysicsAssetFactoryNode>(BaseNodeContainer, NAME_None);
 		if (ensure(SkeletalMeshFactoryNode))
 		{
-			const FString PhysicsAssetUid = UInterchangeFactoryBaseNode::BuildFactoryNodeUid(SkeletalMeshUid_MeshNamePart + SkeletonUid);
+			const FString PhysicsAssetUid = UInterchangeFactoryBaseNode::BuildFactoryNodeUid(SkeletalMeshUid_MeshNamePart + SkeletonUid + TEXT("_PhysicsAsset"));
 			const FString PhysicsAssetDisplayLabel = DisplayLabel + TEXT("_PhysicsAsset");
 			PhysicsAssetFactoryNode->InitializePhysicsAssetNode(PhysicsAssetUid, PhysicsAssetDisplayLabel, UPhysicsAsset::StaticClass()->GetName());
 			PhysicsAssetFactoryNode->SetCustomSkeletalMeshUid(SkeletalMeshUid);
diff --git a/Engine/Plugins/Experimental/Interchange/Runtime/Source/Pipelines/Public/InterchangeGenericAssetsPipeline.h b/Engine/Plugins/Experimental/Interchange/Runtime/Source/Pipelines/Public/InterchangeGenericAssetsPipeline.h
--- a/Engine/Plugins/Experimental/Interchange/Runtime/Source/Pipelines/Public/InterchangeGenericAssetsPipeline.h
+++ b/Engine/Plugins/Experimental/Interchange/Runtime/Source/Pipelines/Public/InterchangeGenericAssetsPipeline.h
@@ -23,7 +23,7 @@ class UInterchangeGenericTexturePipeline;
  * All shared import options between mesh type should be added here.
  *
  */
-UCLASS(BlueprintType, Experimental)
+UCLASS(BlueprintType, editinlinenew, Experimental)
 class INTERCHANGEPIPELINES_API UInterchangeGenericAssetsPipeline : public UInterchangePipelineBase
 {
 	GENERATED_BODY()
@@ -115,5 +115,3 @@ private:
 	TArray<const UInterchangeSourceData*> SourceDatas;
 	
 };
-
-
diff --git a/Engine/Plugins/Experimental/Interchange/Runtime/Source/Pipelines/Public/InterchangeGenericAssetsPipelineSharedSettings.h b/Engine/Plugins/Experimental/Interchange/Runtime/Source/Pipelines/Public/InterchangeGenericAssetsPipelineSharedSettings.h
--- a/Engine/Plugins/Experimental/Interchange/Runtime/Source/Pipelines/Public/InterchangeGenericAssetsPipelineSharedSettings.h
+++ b/Engine/Plugins/Experimental/Interchange/Runtime/Source/Pipelines/Public/InterchangeGenericAssetsPipelineSharedSettings.h
@@ -44,7 +44,7 @@ enum class EInterchangeVertexColorImportOption : uint8
 };
 
 
-UCLASS(BlueprintType, Experimental, hidedropdown)
+UCLASS(BlueprintType, hidedropdown, Experimental)
 class INTERCHANGEPIPELINES_API UInterchangeGenericCommonMeshesProperties : public UInterchangePipelineBase
 {
 	GENERATED_BODY()
@@ -72,7 +72,7 @@ public:
 	FColor VertexOverrideColor;
 };
 
-UCLASS(BlueprintType, Experimental, hidedropdown)
+UCLASS(BlueprintType, hidedropdown, Experimental)
 class INTERCHANGEPIPELINES_API UInterchangeGenericCommonSkeletalMeshesAndAnimationsProperties : public UInterchangePipelineBase
 {
 	GENERATED_BODY()
diff --git a/Engine/Plugins/Experimental/Interchange/Runtime/Source/Pipelines/Public/InterchangeGenericMaterialPipeline.h b/Engine/Plugins/Experimental/Interchange/Runtime/Source/Pipelines/Public/InterchangeGenericMaterialPipeline.h
--- a/Engine/Plugins/Experimental/Interchange/Runtime/Source/Pipelines/Public/InterchangeGenericMaterialPipeline.h
+++ b/Engine/Plugins/Experimental/Interchange/Runtime/Source/Pipelines/Public/InterchangeGenericMaterialPipeline.h
@@ -29,7 +29,7 @@ enum class EInterchangeMaterialImportOption : uint8
 	ImportAsMaterialInstances,
 };
 
-UCLASS(BlueprintType, Experimental, hidedropdown)
+UCLASS(BlueprintType, hidedropdown, Experimental)
 class INTERCHANGEPIPELINES_API UInterchangeGenericMaterialPipeline : public UInterchangePipelineBase
 {
 	GENERATED_BODY()
diff --git a/Engine/Plugins/Experimental/Interchange/Runtime/Source/Pipelines/Public/InterchangeGenericMeshPipeline.h b/Engine/Plugins/Experimental/Interchange/Runtime/Source/Pipelines/Public/InterchangeGenericMeshPipeline.h
--- a/Engine/Plugins/Experimental/Interchange/Runtime/Source/Pipelines/Public/InterchangeGenericMeshPipeline.h
+++ b/Engine/Plugins/Experimental/Interchange/Runtime/Source/Pipelines/Public/InterchangeGenericMeshPipeline.h
@@ -24,7 +24,7 @@ class UInterchangeStaticMeshLodDataNode;
 class UPhysicsAsset;
 
 /* Hide drop down will make sure the class is not showing in the class picker */
-UCLASS(BlueprintType, Experimental, hidedropdown)
+UCLASS(BlueprintType, hidedropdown, Experimental)
 class INTERCHANGEPIPELINES_API UInterchangeGenericMeshPipeline : public UInterchangePipelineBase
 {
 	GENERATED_BODY()
@@ -69,7 +69,7 @@ public:
 	//////	Static Meshes Build settings Properties //////
 
 	/** If enabled this option will allow you to use Nanite rendering at runtime. Can only be used with simple opaque materials. */
-	UPROPERTY(EditAnywhere, BlueprintReadWrite, Category = "Static Meshes|Build Settings")
+	UPROPERTY(EditAnywhere, BlueprintReadWrite, Category = "Static Meshes", meta = (SubCategory = "Build"))
 	bool bBuildNanite = false;
 	
 	//////	SKELETAL_MESHES_CATEGORY Properties //////
@@ -108,6 +108,8 @@ public:
 
 	virtual void AdjustSettingsForReimportType(EInterchangeReimportType ImportType, TObjectPtr<UObject> ReimportAsset) override;
 
+	virtual void PreDialogCleanup(const FName PipelineStackName) override;
+
 protected:
 	virtual void ExecutePreImportPipeline(UInterchangeBaseNodeContainer* InBaseNodeContainer, const TArray<UInterchangeSourceData*>& InSourceDatas) override;
 
diff --git a/Engine/Plugins/Experimental/Interchange/Runtime/Source/Pipelines/Public/InterchangeGenericScenesPipeline.h b/Engine/Plugins/Experimental/Interchange/Runtime/Source/Pipelines/Public/InterchangeGenericScenesPipeline.h
--- a/Engine/Plugins/Experimental/Interchange/Runtime/Source/Pipelines/Public/InterchangeGenericScenesPipeline.h
+++ b/Engine/Plugins/Experimental/Interchange/Runtime/Source/Pipelines/Public/InterchangeGenericScenesPipeline.h
@@ -15,7 +15,7 @@ class UInterchangeActorFactoryNode;
 class UInterchangeSceneNode;
 
 
-UCLASS(BlueprintType, Experimental)
+UCLASS(BlueprintType, editinlinenew, Experimental)
 class INTERCHANGEPIPELINES_API UInterchangeGenericLevelPipeline : public UInterchangePipelineBase
 {
 	GENERATED_BODY()
diff --git a/Engine/Plugins/Experimental/Interchange/Runtime/Source/Pipelines/Public/InterchangeGenericTexturePipeline.h b/Engine/Plugins/Experimental/Interchange/Runtime/Source/Pipelines/Public/InterchangeGenericTexturePipeline.h
--- a/Engine/Plugins/Experimental/Interchange/Runtime/Source/Pipelines/Public/InterchangeGenericTexturePipeline.h
+++ b/Engine/Plugins/Experimental/Interchange/Runtime/Source/Pipelines/Public/InterchangeGenericTexturePipeline.h
@@ -13,7 +13,7 @@
 class UInterchangeTextureFactoryNode;
 class UInterchangeTextureNode;
 
-UCLASS(BlueprintType, Experimental, hidedropdown)
+UCLASS(BlueprintType, hidedropdown, Experimental)
 class INTERCHANGEPIPELINES_API UInterchangeGenericTexturePipeline : public UInterchangePipelineBase
 {
 	GENERATED_BODY()
diff --git a/Engine/Plugins/MeshPainting/Source/MeshPaintEditorMode/Private/MeshPaintModeHelpers.cpp b/Engine/Plugins/MeshPainting/Source/MeshPaintEditorMode/Private/MeshPaintModeHelpers.cpp
--- a/Engine/Plugins/MeshPainting/Source/MeshPaintEditorMode/Private/MeshPaintModeHelpers.cpp
+++ b/Engine/Plugins/MeshPainting/Source/MeshPaintEditorMode/Private/MeshPaintModeHelpers.cpp
@@ -49,6 +49,7 @@
 #include "InterchangeAssetImportData.h"
 #include "InterchangeGenericAssetsPipeline.h"
 #include "InterchangeGenericMeshPipeline.h"
+#include "InterchangePythonPipelineBase.h"
 
 
 void UMeshPaintModeSubsystem::SetViewportColorMode(EMeshPaintDataColorViewMode ColorViewMode, FEditorViewportClient* ViewportClient)
@@ -318,14 +319,24 @@ void UMeshPaintModeSubsystem::ImportVertexColorsToSkeletalMesh(USkeletalMesh* Sk
 		UInterchangeAssetImportData* InterchangeAssetImportData = Cast<UInterchangeAssetImportData>(SkeletalMesh->GetAssetImportData());
 		if (InterchangeAssetImportData)
 		{
-			for (TObjectPtr<UInterchangePipelineBase> PipelineBase : InterchangeAssetImportData->Pipelines)
+			for (TObjectPtr<UObject> PipelineBase : InterchangeAssetImportData->Pipelines)
 			{
-				if (UInterchangeGenericAssetsPipeline* GenericAssetPipeline = Cast<UInterchangeGenericAssetsPipeline>(PipelineBase.Get()))
+				UInterchangeGenericAssetsPipeline* GenericAssetPipeline = Cast<UInterchangeGenericAssetsPipeline>(PipelineBase.Get());
+
+				if (!GenericAssetPipeline)
+				{
+					if (UInterchangePythonPipelineAsset* PythonPipelineAsset = Cast<UInterchangePythonPipelineAsset>(PipelineBase.Get()))
+					{
+						GenericAssetPipeline = Cast<UInterchangeGenericAssetsPipeline>(PythonPipelineAsset->GeneratedPipeline);
+					}
+				}
+
+				if (GenericAssetPipeline)
 				{
 					if (!GenericAssetPipeline->CommonMeshesProperties.IsNull() && GenericAssetPipeline->CommonMeshesProperties->VertexColorImportOption != EInterchangeVertexColorImportOption::IVCIO_Ignore)
 					{
-						InterchangeAssetImportData->SetFlags(RF_Transactional);
-						InterchangeAssetImportData->Modify();
+						GenericAssetPipeline->SetFlags(RF_Transactional);
+						GenericAssetPipeline->Modify();
 						GenericAssetPipeline->CommonMeshesProperties->VertexColorImportOption = EInterchangeVertexColorImportOption::IVCIO_Ignore;
 					}
 				}
diff --git a/Engine/Plugins/MeshPainting/Source/MeshPaintingToolset/Private/MeshPaintSkeletalMeshAdapter.cpp b/Engine/Plugins/MeshPainting/Source/MeshPaintingToolset/Private/MeshPaintSkeletalMeshAdapter.cpp
--- a/Engine/Plugins/MeshPainting/Source/MeshPaintingToolset/Private/MeshPaintSkeletalMeshAdapter.cpp
+++ b/Engine/Plugins/MeshPainting/Source/MeshPaintingToolset/Private/MeshPaintSkeletalMeshAdapter.cpp
@@ -12,6 +12,7 @@
 #include "InterchangeAssetImportData.h"
 #include "InterchangeGenericAssetsPipeline.h"
 #include "InterchangeGenericMeshPipeline.h"
+#include "InterchangePythonPipelineBase.h"
 
 DEFINE_LOG_CATEGORY_STATIC(LogMeshPaintSkeletalMeshAdapter, Log, All);
 //////////////////////////////////////////////////////////////////////////
@@ -398,9 +399,19 @@ void FMeshPaintSkeletalMeshComponentAdapter::PreEdit()
 		UInterchangeAssetImportData* InterchangeAssetImportData = Cast<UInterchangeAssetImportData>(ReferencedSkeletalMesh->GetAssetImportData());
 		if (InterchangeAssetImportData)
 		{
-			for (TObjectPtr<UInterchangePipelineBase> PipelineBase : InterchangeAssetImportData->Pipelines)
+			for (TObjectPtr<UObject> PipelineBase : InterchangeAssetImportData->Pipelines)
 			{
-				if (UInterchangeGenericAssetsPipeline* GenericAssetPipeline = Cast<UInterchangeGenericAssetsPipeline>(PipelineBase.Get()))
+				UInterchangeGenericAssetsPipeline* GenericAssetPipeline = Cast<UInterchangeGenericAssetsPipeline>(PipelineBase.Get());
+
+				if (!GenericAssetPipeline)
+				{
+					if (UInterchangePythonPipelineAsset* PythonPipelineAsset = Cast<UInterchangePythonPipelineAsset>(PipelineBase.Get()))
+					{
+						GenericAssetPipeline = Cast<UInterchangeGenericAssetsPipeline>(PythonPipelineAsset->GeneratedPipeline);
+					}
+				}
+				
+				if (GenericAssetPipeline)
 				{
 					if (!GenericAssetPipeline->CommonMeshesProperties.IsNull() && GenericAssetPipeline->CommonMeshesProperties->VertexColorImportOption != EInterchangeVertexColorImportOption::IVCIO_Ignore)
 					{
diff --git a/Engine/Source/Runtime/Experimental/Interchange/Core/Private/InterchangePipelineBase.cpp b/Engine/Source/Runtime/Experimental/Interchange/Core/Private/InterchangePipelineBase.cpp
--- a/Engine/Source/Runtime/Experimental/Interchange/Core/Private/InterchangePipelineBase.cpp
+++ b/Engine/Source/Runtime/Experimental/Interchange/Core/Private/InterchangePipelineBase.cpp
@@ -16,7 +16,7 @@ namespace UE
 		{
 			FString CreateConfigSectionName(const FName PipelineStackName, UClass* PipelineClass)
 			{
-				FString Section = TEXT("Interchange__StackName_") + PipelineStackName.ToString() + TEXT("__PipelineClassName_") + PipelineClass->GetName();
+				FString Section = TEXT("Interchange_StackName__") + PipelineStackName.ToString() + TEXT("__PipelineClassName_") + PipelineClass->GetName();
 				return Section;
 			}
 		}
@@ -24,6 +24,47 @@ namespace UE
 }
 
 void UInterchangePipelineBase::LoadSettings(const FName PipelineStackName)
+{
+	LoadSettingsInternal(PipelineStackName, GEditorPerProjectIni, LockedProperties);
+}
+
+void UInterchangePipelineBase::SaveSettings(const FName PipelineStackName)
+{
+	SaveSettingsInternal(PipelineStackName, GEditorPerProjectIni);
+}
+
+bool UInterchangePipelineBase::SetLockedPropertyStatus(const FName PropertyPath, bool bLocked)
+{
+	ensure(bAllowLockedPropertiesEdition);
+
+	if (!bLocked)
+	{
+		LockedProperties.Remove(PropertyPath);
+	}
+	else
+	{
+		bool& bLockStatus = LockedProperties.FindOrAdd(PropertyPath);
+		bLockStatus = bLocked;
+	}
+
+	return true;
+}
+
+bool UInterchangePipelineBase::GetLockedPropertyStatus(const FName PropertyPath) const
+{
+	if (const bool* bLockStatus = LockedProperties.Find(PropertyPath))
+	{
+		return *bLockStatus;
+	}
+	return false;
+}
+
+FName UInterchangePipelineBase::GetLockedPropertiesPropertyName()
+{
+	return GET_MEMBER_NAME_CHECKED(UInterchangePipelineBase, LockedProperties);
+}
+
+void UInterchangePipelineBase::LoadSettingsInternal(const FName PipelineStackName, const FString& ConfigFilename, TMap<FName, bool>& ParentLockedProperties)
 {
 	int32 PortFlags = 0;
 	UClass* Class = this->GetClass();
@@ -35,25 +76,39 @@ void UInterchangePipelineBase::LoadSettings(const FName PipelineStackName)
 			continue;
 		}
 
-		FString Section = UE::Interchange::PipelinePrivate::CreateConfigSectionName(PipelineStackName, Class);
+		FString SectionName = UE::Interchange::PipelinePrivate::CreateConfigSectionName(PipelineStackName, Class);
 		FString Key = Property->GetName();
+		const FName PropertyName = Property->GetFName();
+		const FName PropertyPath = FName(Property->GetPathName());
+		if (PropertyName == GET_MEMBER_NAME_CHECKED(UInterchangePipelineBase, LockedProperties))
+		{
+			continue;
+		}
+
+		if (const bool* LockStatus = ParentLockedProperties.Find(PropertyPath))
+		{
+			if (*LockStatus)
+			{
+				//Skip this locked property
+				continue;
+			}
+		}
 
 		const bool bIsPropertyInherited = Property->GetOwnerClass() != Class;
 		UObject* SuperClassDefaultObject = Class->GetSuperClass()->GetDefaultObject();
 
-		const FString& PropFileName = GEditorPerProjectIni;
 		FObjectProperty* SubObject = CastField<FObjectProperty>(Property);
 		FArrayProperty* Array = CastField<FArrayProperty>(Property);
 		if (Array)
 		{
 			const bool bForce = false;
 			const bool bConst = true;
-			FConfigSection* Sec = GConfig->GetSectionPrivate(*Section, bForce, bConst, *GEditorPerProjectIni);
-			if (Sec != nullptr)
+			FConfigSection* Section = GConfig->GetSectionPrivate(*SectionName, bForce, bConst, ConfigFilename);
+			if (Section != nullptr)
 			{
 				TArray<FConfigValue> List;
 				const FName KeyName(*Key, FNAME_Find);
-				Sec->MultiFind(KeyName, List);
+				Section->MultiFind(KeyName, List);
 
 				FScriptArrayHelper_InContainer ArrayHelper(Array, this);
 				// Only override default properties if there is something to override them with.
@@ -80,7 +135,7 @@ void UInterchangePipelineBase::LoadSettings(const FName PipelineStackName)
 						{
 							break;
 						}
-						ElementValue = Sec->Find(IndexedName);
+						ElementValue = Section->Find(IndexedName);
 
 						// If found, import the element
 						if (ElementValue != nullptr)
@@ -100,7 +155,7 @@ void UInterchangePipelineBase::LoadSettings(const FName PipelineStackName)
 			// Load the settings if the referenced pipeline is a subobject of ours
 			if (SubPipeline->IsInOuter(this))
 			{
-				SubPipeline->LoadSettings(PipelineStackName);
+				SubPipeline->LoadSettingsInternal(PipelineStackName, ConfigFilename, ParentLockedProperties);
 			}
 		}
 		else
@@ -113,7 +168,7 @@ void UInterchangePipelineBase::LoadSettings(const FName PipelineStackName)
 				}
 
 				FString Value;
-				bool bFoundValue = GConfig->GetString(*Section, *Key, Value, *GEditorPerProjectIni);
+				bool bFoundValue = GConfig->GetString(*SectionName, *Key, Value, ConfigFilename);
 
 				if (bFoundValue)
 				{
@@ -128,7 +183,7 @@ void UInterchangePipelineBase::LoadSettings(const FName PipelineStackName)
 	}
 }
 
-void UInterchangePipelineBase::SaveSettings(const FName PipelineStackName)
+void UInterchangePipelineBase::SaveSettingsInternal(const FName PipelineStackName, const FString& ConfigFilename)
 {
 	int32 PortFlags = 0;
 	UClass* Class = this->GetClass();
@@ -140,8 +195,13 @@ void UInterchangePipelineBase::SaveSettings(const FName PipelineStackName)
 			continue;
 		}
 
-		FString Section = UE::Interchange::PipelinePrivate::CreateConfigSectionName(PipelineStackName, Class);
+		FString SectionName = UE::Interchange::PipelinePrivate::CreateConfigSectionName(PipelineStackName, Class);
 		FString Key = Property->GetName();
+		const FName PropertyName = Property->GetFName();
+		if (PropertyName == GET_MEMBER_NAME_CHECKED(UInterchangePipelineBase, LockedProperties))
+		{
+			continue;
+		}
 
 		const bool bIsPropertyInherited = Property->GetOwnerClass() != Class;
 		UObject* SuperClassDefaultObject = Class->GetSuperClass()->GetDefaultObject();
@@ -151,16 +211,16 @@ void UInterchangePipelineBase::SaveSettings(const FName PipelineStackName)
 		{
 			const bool bForce = true;
 			const bool bConst = false;
-			FConfigSection* Sec = GConfig->GetSectionPrivate(*Section, bForce, bConst, *GEditorPerProjectIni);
-			check(Sec);
-			Sec->Remove(*Key);
+			FConfigSection* Section = GConfig->GetSectionPrivate(*SectionName, bForce, bConst, ConfigFilename);
+			check(Section);
+			Section->Remove(*Key);
 
 			FScriptArrayHelper_InContainer ArrayHelper(Array, this);
 			for (int32 i = 0; i < ArrayHelper.Num(); i++)
 			{
 				FString	Buffer;
 				Array->Inner->ExportTextItem_Direct(Buffer, ArrayHelper.GetRawPtr(i), ArrayHelper.GetRawPtr(i), this, PortFlags);
-				Sec->Add(*Key, *Buffer);
+				Section->Add(*Key, *Buffer);
 			}
 		}
 		else if (UInterchangePipelineBase* SubPipeline = SubObject ? Cast<UInterchangePipelineBase>(SubObject->GetObjectPropertyValue_InContainer(this)) : nullptr)
@@ -168,7 +228,7 @@ void UInterchangePipelineBase::SaveSettings(const FName PipelineStackName)
 			// Save the settings if the referenced pipeline is a subobject of ours
 			if (SubPipeline->IsInOuter(this))
 			{
-				SubPipeline->SaveSettings(PipelineStackName);
+				SubPipeline->SaveSettingsInternal(PipelineStackName, ConfigFilename);
 			}
 		}
 		else
@@ -184,10 +244,9 @@ void UInterchangePipelineBase::SaveSettings(const FName PipelineStackName)
 
 				FString	Value;
 				Property->ExportText_InContainer(Index, Value, this, this, this, PortFlags);
-				GConfig->SetString(*Section, *Key, *Value, *GEditorPerProjectIni);
+				GConfig->SetString(*SectionName, *Key, *Value, ConfigFilename);
 			}
 		}
 	}
 	GConfig->Flush(0);
 }
-
diff --git a/Engine/Source/Runtime/Experimental/Interchange/Core/Public/InterchangeFactoryBase.h b/Engine/Source/Runtime/Experimental/Interchange/Core/Public/InterchangeFactoryBase.h
--- a/Engine/Source/Runtime/Experimental/Interchange/Core/Public/InterchangeFactoryBase.h
+++ b/Engine/Source/Runtime/Experimental/Interchange/Core/Public/InterchangeFactoryBase.h
@@ -140,6 +140,7 @@ public:
 		FString NodeUniqueID;
 		UInterchangeBaseNodeContainer* NodeContainer = nullptr;
 		TArray<UInterchangePipelineBase*> Pipelines;
+		TArray<UObject*> OriginalPipelines;
 		
 		bool bIsReimport  = false;
 ;
diff --git a/Engine/Source/Runtime/Experimental/Interchange/Core/Public/InterchangePipelineBase.h b/Engine/Source/Runtime/Experimental/Interchange/Core/Public/InterchangePipelineBase.h
--- a/Engine/Source/Runtime/Experimental/Interchange/Core/Public/InterchangePipelineBase.h
+++ b/Engine/Source/Runtime/Experimental/Interchange/Core/Public/InterchangePipelineBase.h
@@ -31,7 +31,7 @@ enum class EInterchangeReimportType : uint8
 	AssetAlternateSkinningReimport,
 };
 
-UCLASS(BlueprintType, Blueprintable, Experimental, EditInlineNew, Abstract)
+UCLASS(BlueprintType, Blueprintable, editinlinenew, Experimental, Abstract)
 class INTERCHANGECORE_API UInterchangePipelineBase : public UObject
 {
 	GENERATED_BODY()
@@ -168,6 +168,22 @@ public:
 		Results = InResults;
 	}
 
+	bool SetLockedPropertyStatus(const FName PropertyPath, bool bLocked);
+	bool GetLockedPropertyStatus(const FName PropertyPath) const;
+
+	/**
+	 * Return the property name of the LockedProperties map
+	 */
+	static FName GetLockedPropertiesPropertyName();
+
+	/**
+	 * If true, the property editor for this pipeline instance will allow locked properties edition.
+	 * If false, the property editor for this pipeline instance will set locked properties in read only.
+	 * 
+	 * Note: If you open in the content browser a pipeline asset you will be able to edit locked properties.
+	 *       If you import a file with interchange, the import dialog will show locked properties in read only mode.
+	 */
+	bool bAllowLockedPropertiesEdition = true;
 
 protected:
 
@@ -209,7 +225,14 @@ protected:
 	{
 	}
 
+	void LoadSettingsInternal(const FName PipelineStackName, const FString& ConfigFilename, TMap<FName, bool>& ParentLockedProperties);
+
+	void SaveSettingsInternal(const FName PipelineStackName, const FString& ConfigFilename);
 
 	UPROPERTY()
 	TObjectPtr<UInterchangeResultsContainer> Results;
+
+	/* Map of property path and lock status. Any properties that have a true lock status will be readonly when showing the import dialog */
+	UPROPERTY(BlueprintReadWrite, Category = "Pipeline Properties Manager")
+	TMap<FName, bool> LockedProperties;
 };
diff --git a/Engine/Source/Runtime/Experimental/Interchange/Core/Public/Nodes/InterchangeBaseNodeContainer.h b/Engine/Source/Runtime/Experimental/Interchange/Core/Public/Nodes/InterchangeBaseNodeContainer.h
--- a/Engine/Source/Runtime/Experimental/Interchange/Core/Public/Nodes/InterchangeBaseNodeContainer.h
+++ b/Engine/Source/Runtime/Experimental/Interchange/Core/Public/Nodes/InterchangeBaseNodeContainer.h
@@ -183,7 +183,7 @@ private:
 	UInterchangeBaseNode* GetNodeChildrenInternal(const FString& NodeUniqueID, int32 ChildIndex);
 
 	/** Flat List of the nodes. Since the nodes are variable size, we store a pointer. */
-	UPROPERTY()
+	UPROPERTY(VisibleAnywhere, Category = "Interchange | Node Container")
 	TMap<FString, TObjectPtr<UInterchangeBaseNode> > Nodes;
 
 	mutable TMap<FString, TArray<FString> > ChildrenCache;
diff --git a/Engine/Source/Runtime/Experimental/Interchange/Engine/InterchangeEngine.Build.cs b/Engine/Source/Runtime/Experimental/Interchange/Engine/InterchangeEngine.Build.cs
--- a/Engine/Source/Runtime/Experimental/Interchange/Engine/InterchangeEngine.Build.cs
+++ b/Engine/Source/Runtime/Experimental/Interchange/Engine/InterchangeEngine.Build.cs
@@ -23,6 +23,8 @@ public class InterchangeEngine : ModuleRules
 				"Slate",
 				"DeveloperSettings",
 				"SkeletalMeshUtilitiesCommon",
+				"Json",
+				"JsonUtilities",
 			}
 		);
 	}
diff --git a/Engine/Source/Runtime/Experimental/Interchange/Engine/Private/InterchangeImportCommon.cpp b/Engine/Source/Runtime/Experimental/Interchange/Engine/Private/InterchangeImportCommon.cpp
--- a/Engine/Source/Runtime/Experimental/Interchange/Engine/Private/InterchangeImportCommon.cpp
+++ b/Engine/Source/Runtime/Experimental/Interchange/Engine/Private/InterchangeImportCommon.cpp
@@ -5,6 +5,7 @@
 #include "EditorFramework/AssetImportData.h"
 #include "InterchangeAssetImportData.h"
 #include "InterchangePipelineBase.h"
+#include "InterchangePythonPipelineBase.h"
 #include "InterchangeSourceData.h"
 #include "Nodes/InterchangeBaseNodeContainer.h"
 #include "Nodes/InterchangeFactoryBaseNode.h"
@@ -49,9 +50,9 @@ namespace UE
 				FObjectDuplicationParameters DupParam(Parameters.NodeContainer, AssetImportData);
 				AssetImportData->NodeContainer = CastChecked<UInterchangeBaseNodeContainer>(StaticDuplicateObjectEx(DupParam));
 				AssetImportData->Pipelines.Reset();
-				for (const UInterchangePipelineBase* Pipeline : Parameters.Pipelines)
+				for (const UObject* Pipeline : Parameters.Pipelines)
 				{
-					UInterchangePipelineBase* DupPipeline = Cast<UInterchangePipelineBase>(StaticDuplicateObject(Pipeline, AssetImportData));
+					UObject* DupPipeline = Cast<UObject>(StaticDuplicateObject(Pipeline, AssetImportData));
 					if (DupPipeline)
 					{
 						AssetImportData->Pipelines.Add(DupPipeline);
@@ -65,7 +66,7 @@ namespace UE
 																							, const UInterchangeSourceData* InSourceData
 																							, FString InNodeUniqueID
 																							, UInterchangeBaseNodeContainer* InNodeContainer
-																						   , const TArray<UInterchangePipelineBase*>& InPipelines)
+																						   , const TArray<UObject*>& InPipelines)
 			: AssetImportDataOuter(InAssetImportDataOuter)
 			, AssetImportData(InAssetImportData)
 			, SourceData(InSourceData)
@@ -135,7 +136,7 @@ namespace UE
 																					, const UInterchangeSourceData* InSourceData
 																					, FString InNodeUniqueID
 																					, UInterchangeBaseNodeContainer* InNodeContainer
-																					, const TArray<UInterchangePipelineBase*>& InPipelines)
+																					, const TArray<UObject*>& InPipelines)
 			: FUpdateImportAssetDataParameters(InAssetImportDataOuter
 				, InAssetImportData
 				, InSourceData
@@ -220,9 +221,20 @@ namespace UE
 				return false;
 			}
 
-			for (UInterchangePipelineBase* PipelineBase : InterchangeAssetImportData->Pipelines)
+			for (UObject* PipelineObject : InterchangeAssetImportData->Pipelines)
 			{
-				PipelineBase->ScriptedSetReimportSourceIndex(Object->GetClass(), SourceIndex);
+				if (UInterchangePythonPipelineAsset* PythonPipelineAsset = Cast<UInterchangePythonPipelineAsset>(PipelineObject))
+				{
+					if (PythonPipelineAsset->GeneratedPipeline)
+					{
+						PythonPipelineAsset->GeneratedPipeline->ScriptedSetReimportSourceIndex(Object->GetClass(), SourceIndex);
+						PythonPipelineAsset->SetupFromPipeline(PythonPipelineAsset->GeneratedPipeline);
+					}
+				}
+				else if (UInterchangePipelineBase* PipelineBase = Cast<UInterchangePipelineBase>(PipelineObject))
+				{
+					PipelineBase->ScriptedSetReimportSourceIndex(Object->GetClass(), SourceIndex);
+				}
 			}
 			return true;
 		}
diff --git a/Engine/Source/Runtime/Experimental/Interchange/Engine/Private/InterchangeManager.cpp b/Engine/Source/Runtime/Experimental/Interchange/Engine/Private/InterchangeManager.cpp
--- a/Engine/Source/Runtime/Experimental/Interchange/Engine/Private/InterchangeManager.cpp
+++ b/Engine/Source/Runtime/Experimental/Interchange/Engine/Private/InterchangeManager.cpp
@@ -9,9 +9,11 @@
 #include "Engine/Blueprint.h"
 #include "Framework/Notifications/NotificationManager.h"
 #include "InterchangeAssetImportData.h"
+#include "InterchangeBlueprintPipelineBase.h"
 #include "InterchangeFactoryBase.h"
 #include "InterchangeEngineLogPrivate.h"
 #include "InterchangeProjectSettings.h"
+#include "InterchangePythonPipelineBase.h"
 #include "InterchangeSourceData.h"
 #include "InterchangeTranslatorBase.h"
 #include "InterchangeWriterBase.h"
@@ -359,6 +361,49 @@ void UE::Interchange::SanitizeObjectName(FString& ObjectName)
 	}
 }
 
+UInterchangePipelineBase* UE::Interchange::GeneratePipelineInstance(const FSoftObjectPath& PipelineInstance)
+{
+	UObject* ReferenceInstance = PipelineInstance.TryLoad();
+	if (!ReferenceInstance)
+	{
+		return nullptr;
+	}
+	UInterchangePipelineBase* GeneratedPipeline = nullptr;
+	if (const UInterchangeBlueprintPipelineBase* BlueprintPipeline = Cast<UInterchangeBlueprintPipelineBase>(ReferenceInstance))
+	{
+		if (BlueprintPipeline->GeneratedClass.Get())
+		{
+			GeneratedPipeline = NewObject<UInterchangePipelineBase>(GetTransientPackage(), BlueprintPipeline->GeneratedClass);
+		}
+		else
+		{
+			//Log an error because we cannot load the python class, maybe the python script was not loaded
+			UE_LOG(LogInterchangeEngine, Error, TEXT("Cannot generate a pipeline instance because the blueprint %s do not have a valid generated class."), *PipelineInstance.GetAssetPathName().ToString());
+		}
+	}
+	else if (const UInterchangePythonPipelineAsset* PythonPipeline = Cast<UInterchangePythonPipelineAsset>(ReferenceInstance))
+	{
+		if (PythonPipeline->GeneratedPipeline)
+		{
+			GeneratedPipeline = DuplicateObject<UInterchangePipelineBase>(PythonPipeline->GeneratedPipeline.Get(), GetTransientPackage());
+		}
+		else
+		{
+			//Log an error because we cannot load the python class, maybe the python script was not loaded
+			UE_LOG(LogInterchangeEngine, Error, TEXT("Cannot generate a pipeline instance because the Python pipeline asset %s do not have a valid generated pipeline instance."), *PipelineInstance.GetAssetPathName().ToString());
+		}
+	}
+	else if (const UInterchangePipelineBase* DefaultPipeline = Cast<UInterchangePipelineBase>(ReferenceInstance))
+	{
+		GeneratedPipeline = DuplicateObject<UInterchangePipelineBase>(DefaultPipeline, GetTransientPackage());
+	}
+	else
+	{
+		UE_LOG(LogInterchangeEngine, Error, TEXT("Cannot generate a pipeline instance because the pipeline asset %s type is unknown."), *PipelineInstance.GetAssetPathName().ToString());
+	}
+	return GeneratedPipeline;
+}
+
 UInterchangeManager& UInterchangeManager::GetInterchangeManager()
 {
 	static TStrongObjectPtr<UInterchangeManager> InterchangeManager = nullptr;
@@ -751,7 +796,14 @@ UInterchangeManager::ImportInternal(const FString& ContentPath, const UInterchan
 			TArray<UInterchangePipelineBase*> PipelineStack;
 			for (int32 PipelineIndex = 0; PipelineIndex < OriginalAssetImportData->Pipelines.Num(); ++PipelineIndex)
 			{
-				UInterchangePipelineBase* SourcePipeline = OriginalAssetImportData->Pipelines[PipelineIndex];
+				UInterchangePipelineBase* SourcePipeline = Cast<UInterchangePipelineBase>(OriginalAssetImportData->Pipelines[PipelineIndex]);
+				if (!SourcePipeline)
+				{
+					if (UInterchangePythonPipelineAsset* PythonPipelineAsset = Cast<UInterchangePythonPipelineAsset>(OriginalAssetImportData->Pipelines[PipelineIndex]))
+					{
+						SourcePipeline = PythonPipelineAsset->GeneratedPipeline;
+					}
+				}
 				if (SourcePipeline) //Its possible a pipeline doesnt exist anymore so it wont load into memory when we loading the outer asset
 				{
 					//Duplicate the pipeline saved in the asset import data
@@ -791,6 +843,21 @@ UInterchangeManager::ImportInternal(const FString& ContentPath, const UInterchan
 			// Simply move the existing pipeline
 			AsyncHelper->Pipelines = MoveTemp(PipelineStack);
 
+			//Fill the original pipeline array that will be save in the asset import data
+			for (UInterchangePipelineBase* Pipeline : AsyncHelper->Pipelines)
+			{
+				if (UInterchangePythonPipelineBase* PythonPipeline = Cast<UInterchangePythonPipelineBase>(Pipeline))
+				{
+					UInterchangePythonPipelineAsset* PythonPipelineAsset = NewObject<UInterchangePythonPipelineAsset>(GetTransientPackage());
+					PythonPipelineAsset->PythonClass = PythonPipeline->GetClass();
+					PythonPipelineAsset->SetupFromPipeline(PythonPipeline);
+					AsyncHelper->OriginalPipelines.Add(PythonPipelineAsset);
+				}
+				else
+				{
+					AsyncHelper->OriginalPipelines.Add(Pipeline);
+				}
+			}
 		}
 		else
 		{
@@ -811,40 +878,50 @@ UInterchangeManager::ImportInternal(const FString& ContentPath, const UInterchan
 			}
 			if (!bImportCancel)
 			{
-				//Get the latest PipelineStacks (the Dialog can change the CDO)
-				const TMap<FName, FInterchangePipelineStack>& PipelineStacks = InterchangeProjectSettings->PipelineStacks;
-				if (!PipelineStacks.Contains(PipelineStackName))
+				if (!DefaultPipelineStacks.Contains(PipelineStackName))
 				{
-					if (PipelineStacks.Contains(GetDefaultPipelineStackName()))
+					if (DefaultPipelineStacks.Contains(GetDefaultPipelineStackName()))
 					{
 						PipelineStackName = GetDefaultPipelineStackName();
 					}
 					else
 					{
 						//Log an error, we cannot import asset without a valid pipeline, we will use the first available pipeline
-						for (const TPair<FName, FInterchangePipelineStack>& PipelineStack : PipelineStacks)
+						for (const TPair<FName, FInterchangePipelineStack>& PipelineStack : DefaultPipelineStacks)
 						{
 							PipelineStackName = PipelineStack.Key;
 						}
 					}
 				}
 
-				if (PipelineStacks.Contains(PipelineStackName))
+				if (DefaultPipelineStacks.Contains(PipelineStackName))
 				{
 					//use the default pipeline
-					const FInterchangePipelineStack& PipelineStack = PipelineStacks.FindChecked(PipelineStackName);
-					TArray<TSoftClassPtr<UInterchangePipelineBase>> Pipelines = PipelineStack.Pipelines;
+					const FInterchangePipelineStack& PipelineStack = DefaultPipelineStacks.FindChecked(PipelineStackName);
+					const TArray<FSoftObjectPath>& Pipelines = PipelineStack.Pipelines;
 					for (int32 GraphPipelineIndex = 0; GraphPipelineIndex < Pipelines.Num(); ++GraphPipelineIndex)
 					{
 						if (Pipelines[GraphPipelineIndex].IsValid())
 						{
-							UClass* PipelineClass = Pipelines[GraphPipelineIndex].LoadSynchronous();
-							if (PipelineClass)
+							if(UInterchangePipelineBase* GeneratedPipeline = UE::Interchange::GeneratePipelineInstance(Pipelines[GraphPipelineIndex]))
 							{
-								UInterchangePipelineBase* GeneratedPipeline = NewObject<UInterchangePipelineBase>(GetTransientPackage(), PipelineClass, NAME_None, RF_NoFlags);
+								GeneratedPipeline->bAllowLockedPropertiesEdition = false;
 								//Load the settings for this pipeline
+								//The dialog is saving the settings for all default pipelines
 								GeneratedPipeline->LoadSettings(PipelineStackName);
 								AsyncHelper->Pipelines.Add(GeneratedPipeline);
+
+								//We need to save the python pipeline asset because we cannot save an asset create with a python class
+								UObject* OriginalPipeline = DuplicateObject<UObject>(Pipelines[GraphPipelineIndex].TryLoad(), GetTransientPackage());
+								if (UInterchangePythonPipelineAsset* PythonPipelineAsset = Cast<UInterchangePythonPipelineAsset>(OriginalPipeline))
+								{
+									PythonPipelineAsset->SetupFromPipeline(Cast<UInterchangePythonPipelineBase>(GeneratedPipeline));
+									AsyncHelper->OriginalPipelines.Add(OriginalPipeline);
+								}
+								else
+								{
+									AsyncHelper->OriginalPipelines.Add(GeneratedPipeline);
+								}
 							}
 						}
 					}
@@ -862,7 +939,9 @@ UInterchangeManager::ImportInternal(const FString& ContentPath, const UInterchan
 		for (int32 GraphPipelineIndex = 0; GraphPipelineIndex < ImportAssetParameters.OverridePipelines.Num(); ++GraphPipelineIndex)
 		{
 			// Duplicate the override pipelines to protect the scripted users form making race conditions
-			AsyncHelper->Pipelines.Add(DuplicateObject<UInterchangePipelineBase>(ImportAssetParameters.OverridePipelines[GraphPipelineIndex], GetTransientPackage()));
+			UInterchangePipelineBase* GeneratedPipeline = DuplicateObject<UInterchangePipelineBase>(ImportAssetParameters.OverridePipelines[GraphPipelineIndex], GetTransientPackage());
+			GeneratedPipeline->bAllowLockedPropertiesEdition = false;
+			AsyncHelper->Pipelines.Add(GeneratedPipeline);
 		}
 	}
 
diff --git a/Engine/Source/Runtime/Experimental/Interchange/Engine/Private/InterchangeMeshUtilities.cpp b/Engine/Source/Runtime/Experimental/Interchange/Engine/Private/InterchangeMeshUtilities.cpp
--- a/Engine/Source/Runtime/Experimental/Interchange/Engine/Private/InterchangeMeshUtilities.cpp
+++ b/Engine/Source/Runtime/Experimental/Interchange/Engine/Private/InterchangeMeshUtilities.cpp
@@ -10,6 +10,7 @@
 #include "InterchangeEngineLogPrivate.h"
 #include "InterchangeManager.h"
 #include "InterchangeProjectSettings.h"
+#include "InterchangePythonPipelineBase.h"
 #include "InterchangeSourceData.h"
 #include "Logging/LogMacros.h"
 #include "UObject/Class.h"
@@ -112,11 +113,22 @@ TFuture<bool> UInterchangeMeshUtilities::InternalImportCustomLodAsync(TSharedPtr
 
 	FImportAssetParameters ImportAssetParameters;
 	ImportAssetParameters.bIsAutomated = true;
-	for (TObjectPtr<UInterchangePipelineBase> SelectedPipeline : InterchangeAssetImportData->Pipelines)
+	for (TObjectPtr<UObject> SelectedPipeline : InterchangeAssetImportData->Pipelines)
 	{
-		UInterchangePipelineBase* GeneratedPipeline = Cast<UInterchangePipelineBase>(StaticDuplicateObject(SelectedPipeline, GetTransientPackage()));
-		GeneratedPipeline->AdjustSettingsForReimportType(ImportType, nullptr);
-		ImportAssetParameters.OverridePipelines.Add(GeneratedPipeline);
+		UInterchangePipelineBase* GeneratedPipeline = nullptr;
+		if (UInterchangePythonPipelineAsset* PythonPipelineAsset = Cast<UInterchangePythonPipelineAsset>(SelectedPipeline))
+		{
+			GeneratedPipeline = Cast<UInterchangePipelineBase>(StaticDuplicateObject(PythonPipelineAsset->GeneratedPipeline, GetTransientPackage()));
+		}
+		else
+		{
+			GeneratedPipeline = Cast<UInterchangePipelineBase>(StaticDuplicateObject(SelectedPipeline, GetTransientPackage()));
+		}
+		if (ensure(GeneratedPipeline))
+		{
+			GeneratedPipeline->AdjustSettingsForReimportType(ImportType, nullptr);
+			ImportAssetParameters.OverridePipelines.Add(GeneratedPipeline);
+		}
 	}
 	FString ImportAssetPath = TEXT("/Engine/TempEditor/Interchange/") + FGuid::NewGuid().ToString(EGuidFormats::Base36Encoded);
 	UE::Interchange::FAssetImportResultRef AssetImportResult = InterchangeManager.ImportAssetAsync(ImportAssetPath, SourceData, ImportAssetParameters);
diff --git a/Engine/Source/Runtime/Experimental/Interchange/Engine/Private/InterchangePythonPipelineBase.cpp b/Engine/Source/Runtime/Experimental/Interchange/Engine/Private/InterchangePythonPipelineBase.cpp
new file mode 100644
--- /dev/null
+++ b/Engine/Source/Runtime/Experimental/Interchange/Engine/Private/InterchangePythonPipelineBase.cpp
@@ -0,0 +1,99 @@
+// Copyright Epic Games, Inc. All Rights Reserved.
+
+#include "InterchangePythonPipelineBase.h"
+
+#include "InterchangePipelineBase.h"
+#include "JsonObjectConverter.h"
+#include "UObject/Object.h"
+#include "UObject/UObjectGlobals.h"
+#include "UObject/ObjectMacros.h"
+#include "UObject/UnrealType.h"
+
+#if WITH_EDITOR
+void UInterchangePythonPipelineBase::PostEditChangeProperty(struct FPropertyChangedEvent& PropertyChangedEvent)
+{
+	//Save the change in the parent of class UInterchangePythonPipelineBase
+	UObject* CurrentOuter = GetOuter();
+	while(CurrentOuter)
+	{
+		if (UInterchangePythonPipelineAsset* PythonPipelineAsset = Cast<UInterchangePythonPipelineAsset>(CurrentOuter))
+		{
+			PythonPipelineAsset->SetupFromPipeline(this, false);
+			break;
+		}
+		if (CurrentOuter->IsA<UPackage>())
+		{
+			break;
+		}
+		CurrentOuter = CurrentOuter->GetOuter();
+	}
+
+}
+#endif
+
+void UInterchangePythonPipelineAsset::PostLoad()
+{
+	Super::PostLoad();
+	GeneratePipeline();
+}
+
+#if WITH_EDITOR
+void UInterchangePythonPipelineAsset::PostEditChangeProperty(FPropertyChangedEvent& PropertyChangedEvent)
+{
+	if (PropertyChangedEvent.GetPropertyName() == GET_MEMBER_NAME_CHECKED(UInterchangePythonPipelineAsset, PythonClass))
+	{
+		//Empty the properties data, the user will have to edit the default properties and the locks again
+		JsonDefaultProperties.Empty();
+		GeneratePipeline();
+	}
+}
+#endif
+
+void UInterchangePythonPipelineAsset::GeneratePipeline()
+{
+	GeneratedPipeline = nullptr;
+	if (UClass* LoadedClass = PythonClass.LoadSynchronous())
+	{
+		GeneratedPipeline = NewObject<UInterchangePythonPipelineBase>(this, LoadedClass);
+
+		TSharedPtr<FJsonObject> RootObject;
+		TSharedRef< TJsonReader<> > JsonReader = TJsonReaderFactory<>::Create(JsonDefaultProperties);
+		if (FJsonSerializer::Deserialize(JsonReader, RootObject))
+		{
+			const TSharedPtr<FJsonObject> JsonPipelineProperties = RootObject->GetObjectField(TEXT("GeneratedPipeline"));
+			FJsonObjectConverter::JsonObjectToUStruct(JsonPipelineProperties.ToSharedRef(), LoadedClass, GeneratedPipeline, 0, 0);
+		}
+	}
+}
+
+void UInterchangePythonPipelineAsset::SetupFromPipeline(const UInterchangePythonPipelineBase* PythonPipeline, const bool bRegeneratePipeline /* = true*/)
+{
+	if (!PythonPipeline)
+	{
+		return;
+	}
+
+	PythonClass = PythonPipeline->GetClass();
+	JsonDefaultProperties.Empty();
+	if (UClass* LoadedClass = PythonClass.LoadSynchronous())
+	{
+		TSharedRef<FJsonObject> RootObject = MakeShareable(new FJsonObject);
+		TSharedRef<FJsonObject> PipelinePropertiesObject = MakeShareable(new FJsonObject);
+		if (FJsonObjectConverter::UStructToJsonObject(LoadedClass, PythonPipeline, PipelinePropertiesObject, 0, 0))
+		{
+			RootObject->SetField(TEXT("GeneratedPipeline"), MakeShareable(new FJsonValueObject(PipelinePropertiesObject)));
+		}
+		//Write the json file
+		FString Json;
+		TSharedRef<TJsonWriter<> > JsonWriter = TJsonWriterFactory<>::Create(&Json, 0);
+		if (FJsonSerializer::Serialize(RootObject, JsonWriter))
+		{
+			JsonDefaultProperties = Json;
+		}
+	}
+	
+	if (bRegeneratePipeline)
+	{
+		GeneratePipeline();
+	}
+}
diff --git a/Engine/Source/Runtime/Experimental/Interchange/Engine/Private/Tasks/InterchangeTaskCompletion.cpp b/Engine/Source/Runtime/Experimental/Interchange/Engine/Private/Tasks/InterchangeTaskCompletion.cpp
--- a/Engine/Source/Runtime/Experimental/Interchange/Engine/Private/Tasks/InterchangeTaskCompletion.cpp
+++ b/Engine/Source/Runtime/Experimental/Interchange/Engine/Private/Tasks/InterchangeTaskCompletion.cpp
@@ -70,6 +70,7 @@ void UE::Interchange::FTaskPreCompletion::DoTask(ENamedThreads::Type CurrentThre
 					Arguments.NodeUniqueID = ObjectInfo.FactoryNode ? ObjectInfo.FactoryNode->GetUniqueID() : FString();
 					Arguments.NodeContainer = AsyncHelper->BaseNodeContainers[SourceIndex].Get();
 					Arguments.Pipelines = AsyncHelper->Pipelines;
+					Arguments.OriginalPipelines = AsyncHelper->OriginalPipelines;
 					ObjectInfo.Factory->PreImportPreCompletedCallback(Arguments);
 				}
 
@@ -141,6 +142,7 @@ void UE::Interchange::FTaskPreCompletion::DoTask(ENamedThreads::Type CurrentThre
 					Arguments.NodeUniqueID = ObjectInfo.FactoryNode ? ObjectInfo.FactoryNode->GetUniqueID() : FString();
 					Arguments.NodeContainer = AsyncHelper->BaseNodeContainers[SourceIndex].Get();
 					Arguments.Pipelines = AsyncHelper->Pipelines;
+					Arguments.OriginalPipelines = AsyncHelper->OriginalPipelines;
 					Arguments.bIsReimport = ObjectInfo.bIsReimport;
 					ObjectInfo.Factory->PostImportPreCompletedCallback(Arguments);
 				}
diff --git a/Engine/Source/Runtime/Experimental/Interchange/Engine/Public/InterchangeAssetImportData.h b/Engine/Source/Runtime/Experimental/Interchange/Engine/Public/InterchangeAssetImportData.h
--- a/Engine/Source/Runtime/Experimental/Interchange/Engine/Public/InterchangeAssetImportData.h
+++ b/Engine/Source/Runtime/Experimental/Interchange/Engine/Public/InterchangeAssetImportData.h
@@ -87,6 +87,6 @@ public:
 	FString NodeUniqueID;
 
 	/** List of pipelines use to import an asset */
-	UPROPERTY(VisibleAnywhere, Category = "Interchange | AssetImportData")
-	TArray<TObjectPtr<UInterchangePipelineBase>> Pipelines;
+	UPROPERTY(EditAnywhere, Category = "Interchange | AssetImportData")
+	TArray<TObjectPtr<UObject>> Pipelines;
 };
diff --git a/Engine/Source/Runtime/Experimental/Interchange/Engine/Public/InterchangeBlueprintPipelineBase.h b/Engine/Source/Runtime/Experimental/Interchange/Engine/Public/InterchangeBlueprintPipelineBase.h
new file mode 100644
--- /dev/null
+++ b/Engine/Source/Runtime/Experimental/Interchange/Engine/Public/InterchangeBlueprintPipelineBase.h
@@ -0,0 +1,28 @@
+// Copyright Epic Games, Inc. All Rights Reserved.
+
+#pragma once
+
+#include "Engine/Blueprint.h"
+#include "InterchangePipelineBase.h"
+#include "UObject/Object.h"
+#include "UObject/ObjectMacros.h"
+
+
+#include "InterchangeBlueprintPipelineBase.generated.h"
+
+UCLASS(BlueprintType, Experimental)
+class INTERCHANGEENGINE_API UInterchangeBlueprintPipelineBase : public UBlueprint
+{
+	GENERATED_BODY()
+
+public:
+	UInterchangeBlueprintPipelineBase()
+	{
+		ParentClass = UInterchangePipelineBase::StaticClass();
+		//We must make sure the GeneratedClass is generated after the blueprint is loaded
+		bRecompileOnLoad = true;
+	}
+#if WITH_EDITOR
+	virtual UClass* GetBlueprintClass() const override { return UInterchangePipelineBase::StaticClass(); }
+#endif
+};
\ No newline at end of file
diff --git a/Engine/Source/Runtime/Experimental/Interchange/Engine/Public/InterchangeImportCommon.h b/Engine/Source/Runtime/Experimental/Interchange/Engine/Public/InterchangeImportCommon.h
--- a/Engine/Source/Runtime/Experimental/Interchange/Engine/Public/InterchangeImportCommon.h
+++ b/Engine/Source/Runtime/Experimental/Interchange/Engine/Public/InterchangeImportCommon.h
@@ -32,14 +32,14 @@ namespace UE
 				const UInterchangeSourceData* SourceData = nullptr;
 				FString NodeUniqueID;
 				UInterchangeBaseNodeContainer* NodeContainer;
-				const TArray<UInterchangePipelineBase*> Pipelines;
+				const TArray<UObject*> Pipelines;
 
 				FUpdateImportAssetDataParameters(UObject* InAssetImportDataOuter
 																	, UAssetImportData* InAssetImportData
 																	, const UInterchangeSourceData* InSourceData
 																	, FString InNodeUniqueID
 																	, UInterchangeBaseNodeContainer* InNodeContainer
-																	, const TArray<UInterchangePipelineBase*>& InPipelines);
+																	, const TArray<UObject*>& InPipelines);
 			};
 
 			/**
@@ -68,7 +68,7 @@ namespace UE
 					, const UInterchangeSourceData* InSourceData
 					, FString InNodeUniqueID
 					, UInterchangeBaseNodeContainer* InNodeContainer
-					, const TArray<UInterchangePipelineBase*>& InPipelines);
+					, const TArray<UObject*>& InPipelines);
 			};
 
 			/**
diff --git a/Engine/Source/Runtime/Experimental/Interchange/Engine/Public/InterchangeManager.h b/Engine/Source/Runtime/Experimental/Interchange/Engine/Public/InterchangeManager.h
--- a/Engine/Source/Runtime/Experimental/Interchange/Engine/Public/InterchangeManager.h
+++ b/Engine/Source/Runtime/Experimental/Interchange/Engine/Public/InterchangeManager.h
@@ -180,8 +180,11 @@ namespace UE
 			TArray<UInterchangeSourceData*> SourceDatas;
 			TArray<UInterchangeTranslatorBase*> Translators;
 
-			//Pipelines array is not per source data 
+			//Pipelines array is not per source data
 			TArray<UInterchangePipelineBase*> Pipelines;
+			//The original pipelines asset to save in the asset reimport data. The original pipeline can restore python class member value.
+			//Python class instanced assets cannot be saved, so we have to serialize in json the data to restore it when we do a re-import.
+			TArray<UObject*> OriginalPipelines;
 
 			TArray<FGraphEventRef> TranslatorTasks;
 			TArray<FGraphEventRef> PipelinePreImportTasks;
@@ -236,6 +239,9 @@ namespace UE
 
 		void SanitizeObjectName(FString& ObjectName);
 
+		/* This function take an asset representing a pipeline and generate a UInterchangePipelineBase asset. */
+		INTERCHANGEENGINE_API UInterchangePipelineBase* GeneratePipelineInstance(const FSoftObjectPath& PipelineInstance);
+
 	} //ns interchange
 } //ns UE
 
diff --git a/Engine/Source/Runtime/Experimental/Interchange/Engine/Public/InterchangeProjectSettings.h b/Engine/Source/Runtime/Experimental/Interchange/Engine/Public/InterchangeProjectSettings.h
--- a/Engine/Source/Runtime/Experimental/Interchange/Engine/Public/InterchangeProjectSettings.h
+++ b/Engine/Source/Runtime/Experimental/Interchange/Engine/Public/InterchangeProjectSettings.h
@@ -8,6 +8,7 @@
 #include "InterchangeFilePickerBase.h"
 #include "InterchangePipelineBase.h"
 #include "InterchangePipelineConfigurationBase.h"
+#include "UObject/SoftObjectPath.h"
 
 #include "InterchangeProjectSettings.generated.h"
 
@@ -15,10 +16,10 @@ USTRUCT()
 struct FInterchangePipelineStack
 {
 	GENERATED_BODY()
-
+	
 	/** The starting mesh for the blueprint **/
-	UPROPERTY(EditAnywhere, config, Category = Interchange)
-	TArray<TSoftClassPtr<UInterchangePipelineBase>> Pipelines;
+	UPROPERTY(EditAnywhere, Category = Interchange, meta=(AllowedClasses = "InterchangePipelineBase, InterchangeBlueprintPipelineBase, InterchangePythonPipelineAsset"))
+	TArray<FSoftObjectPath> Pipelines;
 };
 
 UCLASS(config=Engine, meta=(DisplayName=Interchange), MinimalAPI)
diff --git a/Engine/Source/Runtime/Experimental/Interchange/Engine/Public/InterchangePythonPipelineBase.h b/Engine/Source/Runtime/Experimental/Interchange/Engine/Public/InterchangePythonPipelineBase.h
new file mode 100644
--- /dev/null
+++ b/Engine/Source/Runtime/Experimental/Interchange/Engine/Public/InterchangePythonPipelineBase.h
@@ -0,0 +1,72 @@
+// Copyright Epic Games, Inc. All Rights Reserved.
+
+#pragma once
+
+#include "InterchangePipelineBase.h"
+#include "UObject/Object.h"
+#include "UObject/ObjectMacros.h"
+
+
+#include "InterchangePythonPipelineBase.generated.h"
+
+/*
+ * This class represent a python pipeline. It is use by the TSoftClassPtr of the PythonPipeline asset.
+ *
+ */
+UCLASS(BlueprintType, Abstract, Experimental)
+class INTERCHANGEENGINE_API UInterchangePythonPipelineBase : public UInterchangePipelineBase
+{
+	GENERATED_BODY()
+public:
+
+#if WITH_EDITOR
+	virtual void PostEditChangeProperty(struct FPropertyChangedEvent& PropertyChangedEvent) override;
+#endif
+};
+
+
+USTRUCT()
+struct FPropertyData
+{
+	GENERATED_BODY()
+
+	FString DefaultValue;
+	bool bLocked = false;
+};
+
+/*
+ * This class is a helper class for python pipeline. It allow picker to filter correctly the content browser instance we want to support
+ * in the pipeline stack.
+ *
+ * Use this class factory to be able to create an instance of a python pipeline in the content browser.
+ * A python pipeline do not have any locked properties and all default value must be set in the python script.
+ * This restriction exist because python class are transient, so any assets create from a python class cannot be save.
+ * 
+ */
+UCLASS(BlueprintType, Experimental)
+class INTERCHANGEENGINE_API UInterchangePythonPipelineAsset : public UObject
+{
+	GENERATED_BODY()
+
+public:
+	virtual void PostLoad() override;
+#if WITH_EDITOR
+	virtual void PostEditChangeProperty(struct FPropertyChangedEvent& PropertyChangedEvent) override;
+#endif
+
+	void GeneratePipeline();
+
+	void SetupFromPipeline(const UInterchangePythonPipelineBase* PythonPipeline, const bool bRegeneratePipeline = true);
+
+	/** The python class we want to use has a pipeline */
+	UPROPERTY(EditAnywhere, Category = "Interchange|Python")
+	TSoftClassPtr<UInterchangePythonPipelineBase> PythonClass;
+
+	/** The transient pipeline we generate when we load the python pipeline */
+	UPROPERTY(VisibleAnywhere, Transient, Category = "Interchange|Python")
+	TObjectPtr<UInterchangePythonPipelineBase> GeneratedPipeline;
+
+	UPROPERTY(VisibleAnywhere, Category = "Interchange|Python")
+	FString JsonDefaultProperties;
+};
+
