Interchange Configuration Dialog
- Go through the cancel path when closing the dialog window using the X window button.
- Replaced the Import All button with a checkbox and made the Import button a primary button.

diff --git a/Engine/Plugins/Experimental/Interchange/Editor/Source/Pipelines/InterchangeEditorPipelines.Build.cs b/Engine/Plugins/Experimental/Interchange/Editor/Source/Pipelines/InterchangeEditorPipelines.Build.cs
--- a/Engine/Plugins/Experimental/Interchange/Editor/Source/Pipelines/InterchangeEditorPipelines.Build.cs
+++ b/Engine/Plugins/Experimental/Interchange/Editor/Source/Pipelines/InterchangeEditorPipelines.Build.cs
@@ -20,6 +20,7 @@ namespace UnrealBuildTool.Rules
 			PrivateDependencyModuleNames.AddRange(
 				new string[]
 				{
+					"AssetTools",
 					"InputCore",
 					"InterchangePipelines",
 					"MainFrame",
@@ -27,8 +28,8 @@ namespace UnrealBuildTool.Rules
 					"PropertyEditor",
 					"Slate",
 					"SlateCore",
-					"UnrealEd",
-					"AssetTools"
+					"ToolWidgets",
+					"UnrealEd"
 				}
 			);
 		}
diff --git a/Engine/Plugins/Experimental/Interchange/Editor/Source/Pipelines/Private/SInterchangePipelineConfigurationDialog.cpp b/Engine/Plugins/Experimental/Interchange/Editor/Source/Pipelines/Private/SInterchangePipelineConfigurationDialog.cpp
--- a/Engine/Plugins/Experimental/Interchange/Editor/Source/Pipelines/Private/SInterchangePipelineConfigurationDialog.cpp
+++ b/Engine/Plugins/Experimental/Interchange/Editor/Source/Pipelines/Private/SInterchangePipelineConfigurationDialog.cpp
@@ -16,6 +16,7 @@
 #include "Modules/ModuleManager.h"
 #include "PropertyEditorDelegates.h"
 #include "PropertyEditorModule.h"
+#include "SPrimaryButton.h"
 #include "Styling/SlateIconFinder.h"
 #include "Widgets/Images/SImage.h"
 #include "Widgets/Input/SButton.h"
@@ -335,9 +336,10 @@ SInterchangePipelineConfigurationDialog::SInterchangePipelineConfigurationDialog
 
 SInterchangePipelineConfigurationDialog::~SInterchangePipelineConfigurationDialog()
 {
-	PipelineConfigurationTreeView = nullptr;
-	PipelineConfigurationDetailsView = nullptr;
-	OwnerWindow = nullptr;
+	if (TSharedPtr<SWindow> OwnerWindowPinned = OwnerWindow.Pin())
+	{
+		OwnerWindowPinned->GetOnWindowClosedEvent().RemoveAll(this);
+	}
 }
 
 TSharedRef<SBox> SInterchangePipelineConfigurationDialog::SpawnPipelineConfiguration()
@@ -421,6 +423,10 @@ void SInterchangePipelineConfigurationDialog::Construct(const FArguments& InArgs
 	PipelineStack = InArgs._PipelineStack;
 
 	check(OwnerWindow.IsValid());
+	if (TSharedPtr<SWindow> OwnerWindowPinned = OwnerWindow.Pin())
+	{
+		OwnerWindowPinned->GetOnWindowClosedEvent().AddRaw(this, &SInterchangePipelineConfigurationDialog::OnWindowClosed);
+	}
 
 	this->ChildSlot
 	[
@@ -464,36 +470,54 @@ void SInterchangePipelineConfigurationDialog::Construct(const FArguments& InArgs
 			.HAlign(HAlign_Right)
 			.Padding(2)
 			[
-				SNew(SUniformGridPanel)
-				.SlotPadding(2)
-				+ SUniformGridPanel::Slot(0, 0)
+				SNew(SHorizontalBox)
+				+ SHorizontalBox::Slot()
+				.Padding(4.f, 0.f)
+				.AutoWidth()
 				[
 					IDocumentation::Get()->CreateAnchor(FString("Engine/Content/Interchange/PipelineConfiguration"))
 				]
-				+ SUniformGridPanel::Slot(1, 0)
+				+ SHorizontalBox::Slot()
+				.Padding(4.f, 0.f)
+				.AutoWidth()
 				[
-					SNew(SButton)
-					.HAlign(HAlign_Center)
-					.Text(LOCTEXT("InspectorGraphWindow_Cancel", "Cancel"))
-					.OnClicked(this, &SInterchangePipelineConfigurationDialog::OnCloseDialog, ECloseEventType::Cancel)
+					SNew(SHorizontalBox)
+					.ToolTipText(LOCTEXT("InspectorGraphWindow_ReuseSettingsToolTip", "When importing multiple files, keep the same import settings for every file or open the settings dialog for each file."))
+					+ SHorizontalBox::Slot()
+					.AutoWidth()
+					.HAlign(HAlign_Right)
+					.VAlign(VAlign_Center)
+					.Padding(4.f, 0.f)
+					[
+						SNew(STextBlock)
+						.Text(LOCTEXT("InspectorGraphWindow_ReuseSettings", "Use the same settings for subsequent files"))
+					]
+					+ SHorizontalBox::Slot()
+					.Padding(4.f, 0.f)
+					[
+						SAssignNew(UseSameSettingsForAllCheckBox, SCheckBox)
+						.IsChecked(true)
+						.IsEnabled(this, &SInterchangePipelineConfigurationDialog::IsImportButtonEnabled)
+					]
 				]
-				+ SUniformGridPanel::Slot(2, 0)
+				+ SHorizontalBox::Slot()
+				.AutoWidth()
+				.Padding(4.f, 0.f) 
 				[
-					SNew(SButton)
-					.HAlign(HAlign_Center)
-					.Text(LOCTEXT("InspectorGraphWindow_ImportAll", "Import All"))
+					SNew(SPrimaryButton)
+					.Text(LOCTEXT("InspectorGraphWindow_Import", "Import"))
 					.ToolTipText(this, &SInterchangePipelineConfigurationDialog::GetImportButtonTooltip)
 					.IsEnabled(this, &SInterchangePipelineConfigurationDialog::IsImportButtonEnabled)
-					.OnClicked(this, &SInterchangePipelineConfigurationDialog::OnCloseDialog, ECloseEventType::ImportAll)
+					.OnClicked(this, &SInterchangePipelineConfigurationDialog::OnCloseDialog, ECloseEventType::Import)
 				]
-				+ SUniformGridPanel::Slot(3, 0)
+				+ SHorizontalBox::Slot()
+				.AutoWidth()
+				.Padding(4.f, 0.f) 
 				[
 					SNew(SButton)
 					.HAlign(HAlign_Center)
-					.Text(LOCTEXT("InspectorGraphWindow_Import", "Import"))
-					.ToolTipText(this, &SInterchangePipelineConfigurationDialog::GetImportButtonTooltip)
-					.IsEnabled(this, &SInterchangePipelineConfigurationDialog::IsImportButtonEnabled)
-					.OnClicked(this, &SInterchangePipelineConfigurationDialog::OnCloseDialog, ECloseEventType::Import)
+					.Text(LOCTEXT("InspectorGraphWindow_Cancel", "Cancel"))
+					.OnClicked(this, &SInterchangePipelineConfigurationDialog::OnCloseDialog, ECloseEventType::Cancel)
 				]
 			]
 		]
@@ -705,7 +729,7 @@ void SInterchangePipelineConfigurationDialog::RecursiveLoadPipelineSettings(cons
 
 void SInterchangePipelineConfigurationDialog::ClosePipelineConfiguration(const ECloseEventType CloseEventType)
 {
-	if (CloseEventType == ECloseEventType::Cancel)
+	if (CloseEventType == ECloseEventType::Cancel || CloseEventType == ECloseEventType::WindowClosing)
 	{
 		bCanceled = true;
 		bImportAll = false;
@@ -719,16 +743,10 @@ void SInterchangePipelineConfigurationDialog::ClosePipelineConfiguration(const E
 			}
 		}
 	}
-	else if (CloseEventType == ECloseEventType::ImportAll)
+	else //ECloseEventType::Import
 	{
 		bCanceled = false;
-		bImportAll = true;
-	}
-	else
-	{
-		//ECloseEventType::Import
-		bCanceled = false;
-		bImportAll = false;
+		bImportAll = UseSameSettingsForAllCheckBox->IsChecked();
 	}
 
 	if (!bReimport && PipelineConfigurationTreeView)
@@ -743,9 +761,13 @@ void SInterchangePipelineConfigurationDialog::ClosePipelineConfiguration(const E
 	PipelineConfigurationTreeView = nullptr;
 	PipelineConfigurationDetailsView = nullptr;
 
-	if (TSharedPtr<SWindow> OwnerWindowPin = OwnerWindow.Pin())
+	if (CloseEventType != ECloseEventType::WindowClosing)
 	{
-		OwnerWindowPin->RequestDestroyWindow();
+		if (TSharedPtr<SWindow> OwnerWindowPin = OwnerWindow.Pin())
+		{
+			OwnerWindowPin->GetOnWindowClosedEvent().RemoveAll(this);
+			OwnerWindowPin->RequestDestroyWindow();
+		}
 	}
 	OwnerWindow = nullptr;
 }
diff --git a/Engine/Plugins/Experimental/Interchange/Editor/Source/Pipelines/Private/SInterchangePipelineConfigurationDialog.h b/Engine/Plugins/Experimental/Interchange/Editor/Source/Pipelines/Private/SInterchangePipelineConfigurationDialog.h
--- a/Engine/Plugins/Experimental/Interchange/Editor/Source/Pipelines/Private/SInterchangePipelineConfigurationDialog.h
+++ b/Engine/Plugins/Experimental/Interchange/Editor/Source/Pipelines/Private/SInterchangePipelineConfigurationDialog.h
@@ -13,8 +13,9 @@
 #include "Widgets/SCompoundWidget.h"
 #include "Widgets/Views/STreeView.h"
 
-class IDetailsView;
 struct FPropertyAndParent;
+class IDetailsView;
+class SCheckBox;
 
 class FInterchangePipelineStacksTreeNodeItem : protected FGCObject
 {
@@ -86,8 +87,8 @@ protected:
 enum class ECloseEventType : uint8
 {
 	Cancel,
-	ImportAll,
-	Import
+	Import,
+	WindowClosing
 };
 
 class SInterchangePipelineConfigurationDialog : public SCompoundWidget
@@ -106,7 +107,7 @@ public:
 public:
 
 	SInterchangePipelineConfigurationDialog();
-	~SInterchangePipelineConfigurationDialog();
+	virtual ~SInterchangePipelineConfigurationDialog();
 
 	void Construct(const FArguments& InArgs);
 	virtual bool SupportsKeyboardFocus() const override { return true; }
@@ -119,16 +120,17 @@ public:
 		return FReply::Handled();
 	}
 
+	void OnWindowClosed(const TSharedRef<SWindow>& ClosedWindow)
+	{
+		ClosePipelineConfiguration(ECloseEventType::WindowClosing);
+	}
+
 	virtual FReply OnKeyDown(const FGeometry& MyGeometry, const FKeyEvent& InKeyEvent) override;
 
-	bool IsCanceled() { return bCanceled; }
-	bool IsImportAll() { return bImportAll; }
-private:
-	TWeakPtr< SWindow > OwnerWindow;
-	TWeakObjectPtr<UInterchangeSourceData> SourceData;
-	bool bReimport = false;
-	TArray<UInterchangePipelineBase*> PipelineStack;
+	bool IsCanceled() const { return bCanceled; }
+	bool IsImportAll() const { return bImportAll; }
 
+private:
 	TSharedRef<SBox> SpawnPipelineConfiguration();
 	void OnSelectionChanged(TSharedPtr<FInterchangePipelineStacksTreeNodeItem> Item, ESelectInfo::Type SelectionType);
 
@@ -144,10 +146,17 @@ private:
 	void RecursiveSavePipelineSettings(const TSharedPtr<FInterchangePipelineStacksTreeNodeItem>& ParentNode, const int32 PipelineIndex) const;
 	void RecursiveLoadPipelineSettings(const TSharedPtr<FInterchangePipelineStacksTreeNodeItem>& ParentNode, const int32 PipelineIndex) const;
 
+private:
+	TWeakPtr< SWindow > OwnerWindow;
+	TWeakObjectPtr<UInterchangeSourceData> SourceData;
+	TArray<UInterchangePipelineBase*> PipelineStack;
+	
 	//Graph Inspector UI elements
 	TSharedPtr<SInterchangePipelineStacksTreeView> PipelineConfigurationTreeView;
 	TSharedPtr<IDetailsView> PipelineConfigurationDetailsView;
+	TSharedPtr<SCheckBox> UseSameSettingsForAllCheckBox;
 
+	bool bReimport = false;
 	bool bCanceled = false;
 	bool bImportAll = false;
 
