Interchange: Separate the pipeline stacks use by the scene import versus the content import

#rnx

diff --git a/Engine/Config/BaseEngine.ini b/Engine/Config/BaseEngine.ini
--- a/Engine/Config/BaseEngine.ini
+++ b/Engine/Config/BaseEngine.ini
@@ -3052,12 +3052,7 @@ EnvironmentCollisionThickness = 2.f
 EnvironmentCollisionMaxShapes = 32
 
 [/Script/InterchangeEngine.InterchangeProjectSettings]
-PipelineStacks=(("Assets", (Pipelines=("/Interchange/Pipelines/DefaultAssetsPipeline.DefaultAssetsPipeline"))),("Scene", (Pipelines=("/Interchange/Pipelines/DefaultSceneAssetsPipeline.DefaultSceneAssetsPipeline","/Interchange/Pipelines/DefaultSceneLevelPipeline.DefaultSceneLevelPipeline"))))
-DefaultPipelineStack=Assets
-DefaultScenePipelineStack=Scene
-PipelineConfigurationDialogClass=/Script/InterchangeEditorPipelines.InterchangePipelineConfigurationGeneric
-bShowPipelineStacksConfigurationDialog=True
 FilePickerClass=/Script/InterchangeEditorUtilities.InterchangeFilePickerGeneric
-bUseInterchangeWhenImportingIntoLevel=False
 bStaticMeshUseSmoothEdgesIfSmoothingInformationIsMissing=True
-
+ContentImportSettings=(PipelineStacks=(("Assets", (Pipelines=("/Interchange/Pipelines/DefaultAssetsPipeline.DefaultAssetsPipeline")))),DefaultPipelineStack="Assets",PipelineConfigurationDialogClass="/Script/InterchangeEditorPipelines.InterchangePipelineConfigurationGeneric",bShowPipelineStacksConfigurationDialog=True)
+SceneImportSettings=(PipelineStacks=(("Scene", (Pipelines=("/Interchange/Pipelines/DefaultSceneAssetsPipeline.DefaultSceneAssetsPipeline","/Interchange/Pipelines/DefaultSceneLevelPipeline.DefaultSceneLevelPipeline")))),DefaultPipelineStack="Scene",PipelineConfigurationDialogClass="/Script/InterchangeEditorPipelines.InterchangePipelineConfigurationGeneric",bShowPipelineStacksConfigurationDialog=True)
diff --git a/Engine/Plugins/Experimental/Interchange/Editor/Source/Pipelines/Private/InterchangePipelineConfigurationGeneric.cpp b/Engine/Plugins/Experimental/Interchange/Editor/Source/Pipelines/Private/InterchangePipelineConfigurationGeneric.cpp
--- a/Engine/Plugins/Experimental/Interchange/Editor/Source/Pipelines/Private/InterchangePipelineConfigurationGeneric.cpp
+++ b/Engine/Plugins/Experimental/Interchange/Editor/Source/Pipelines/Private/InterchangePipelineConfigurationGeneric.cpp
@@ -21,7 +21,7 @@ EInterchangePipelineConfigurationDialogResult UInterchangePipelineConfigurationG
 	}
 	TSharedRef<SWindow> Window = SNew(SWindow)
 		.ClientSize(FVector2D(1000.f, 650.f))
-		.Title(NSLOCTEXT("Interchange", "PipelineConfigurationGenericTitle", "Interchange Pipeline Configuration"));
+		.Title(NSLOCTEXT("Interchange", "PipelineConfigurationGenericTitleContent", "Interchange Pipeline Configuration (Import Content)"));
 	TSharedPtr<SInterchangePipelineConfigurationDialog> InterchangePipelineConfigurationDialog;
 
 	Window->SetContent
@@ -29,6 +29,7 @@ EInterchangePipelineConfigurationDialogResult UInterchangePipelineConfigurationG
 		SAssignNew(InterchangePipelineConfigurationDialog, SInterchangePipelineConfigurationDialog)
 		.OwnerWindow(Window)
 		.SourceData(SourceData)
+		.bSceneImport(false)
 		.bReimport(false)
 	);
 
@@ -47,6 +48,43 @@ EInterchangePipelineConfigurationDialogResult UInterchangePipelineConfigurationG
 	return EInterchangePipelineConfigurationDialogResult::Import;
 }
 
+EInterchangePipelineConfigurationDialogResult UInterchangePipelineConfigurationGeneric::ShowScenePipelineConfigurationDialog(TWeakObjectPtr<UInterchangeSourceData> SourceData)
+{
+	//Create and show the graph inspector UI dialog
+	TSharedPtr<SWindow> ParentWindow;
+	if (IMainFrameModule* MainFrame = FModuleManager::LoadModulePtr<IMainFrameModule>("MainFrame"))
+	{
+		ParentWindow = MainFrame->GetParentWindow();
+	}
+	TSharedRef<SWindow> Window = SNew(SWindow)
+		.ClientSize(FVector2D(1000.f, 650.f))
+		.Title(NSLOCTEXT("Interchange", "PipelineConfigurationGenericTitleScene", "Interchange Pipeline Configuration (Import Scene)"));
+	TSharedPtr<SInterchangePipelineConfigurationDialog> InterchangePipelineConfigurationDialog;
+
+	Window->SetContent
+	(
+		SAssignNew(InterchangePipelineConfigurationDialog, SInterchangePipelineConfigurationDialog)
+		.OwnerWindow(Window)
+		.SourceData(SourceData)
+		.bSceneImport(true)
+		.bReimport(false)
+	);
+
+	FSlateApplication::Get().AddModalWindow(Window, ParentWindow, false);
+
+	if (InterchangePipelineConfigurationDialog->IsCanceled())
+	{
+		return EInterchangePipelineConfigurationDialogResult::Cancel;
+	}
+
+	if (InterchangePipelineConfigurationDialog->IsImportAll())
+	{
+		return EInterchangePipelineConfigurationDialogResult::ImportAll;
+	}
+
+	return EInterchangePipelineConfigurationDialogResult::Import;
+}
+
 EInterchangePipelineConfigurationDialogResult UInterchangePipelineConfigurationGeneric::ShowReimportPipelineConfigurationDialog(TArray<UInterchangePipelineBase*>& PipelineStack, TWeakObjectPtr<UInterchangeSourceData> SourceData)
 {
 	//Create and show the graph inspector UI dialog
@@ -57,7 +95,7 @@ EInterchangePipelineConfigurationDialogResult UInterchangePipelineConfigurationG
 	}
 	TSharedRef<SWindow> Window = SNew(SWindow)
 		.ClientSize(FVector2D(1000.f, 650.f))
-		.Title(NSLOCTEXT("Interchange", "PipelineConfigurationGenericTitle", "Interchange Pipeline Configuration"));
+		.Title(NSLOCTEXT("Interchange", "PipelineConfigurationGenericTitleReimportContent", "Interchange Pipeline Configuration (Reimport Content)"));
 	TSharedPtr<SInterchangePipelineConfigurationDialog> InterchangePipelineConfigurationDialog;
 
 	Window->SetContent
@@ -65,6 +103,7 @@ EInterchangePipelineConfigurationDialogResult UInterchangePipelineConfigurationG
 		SAssignNew(InterchangePipelineConfigurationDialog, SInterchangePipelineConfigurationDialog)
 		.OwnerWindow(Window)
 		.SourceData(SourceData)
+		.bSceneImport(false)
 		.bReimport(true)
 		.PipelineStack(PipelineStack)
 	);
diff --git a/Engine/Plugins/Experimental/Interchange/Editor/Source/Pipelines/Private/SInterchangePipelineConfigurationDialog.cpp b/Engine/Plugins/Experimental/Interchange/Editor/Source/Pipelines/Private/SInterchangePipelineConfigurationDialog.cpp
--- a/Engine/Plugins/Experimental/Interchange/Editor/Source/Pipelines/Private/SInterchangePipelineConfigurationDialog.cpp
+++ b/Engine/Plugins/Experimental/Interchange/Editor/Source/Pipelines/Private/SInterchangePipelineConfigurationDialog.cpp
@@ -28,6 +28,26 @@
 
 #define LOCTEXT_NAMESPACE "InterchangePipelineConfiguration"
 
+namespace UE::Interchange::Private
+{
+	const FInterchangeImportSettings& GetDefaultImportSettings(const bool bIsSceneImport)
+	{
+		if (bIsSceneImport)
+		{
+			return GetDefault<UInterchangeProjectSettings>()->SceneImportSettings;
+		}
+		else
+		{
+			return GetDefault<UInterchangeProjectSettings>()->ContentImportSettings;
+		}
+	}
+
+	FName GetDefaultPipelineStackName(const bool bIsSceneImport)
+	{
+		return GetDefaultImportSettings(bIsSceneImport).DefaultPipelineStack;
+	}
+}
+
 /************************************************************************/
 /* FInterchangePipelineStacksTreeNodeItem Implementation                    */
 /************************************************************************/
@@ -53,12 +73,25 @@ void SInterchangePipelineStacksTreeView::Construct(const FArguments& InArgs)
 {
 	OnSelectionChangedDelegate = InArgs._OnSelectionChangedDelegate;
 	PipelineStack = InArgs._PipelineStack;
+	bSceneImport = InArgs._bSceneImport;
 	bReimport = InArgs._bReimport && PipelineStack.Num() > 0;
 
 	//Build the FbxNodeInfoPtr tree data
 	const FName ReimportPipelineName = TEXT("ReimportPipeline");
-	const FName& DefaultPipelineStackName = bReimport ? ReimportPipelineName : GetDefault<UInterchangeProjectSettings>()->DefaultPipelineStack;
-	const TMap<FName, FInterchangePipelineStack>& DefaultPipelineStacks = GetDefault<UInterchangeProjectSettings>()->PipelineStacks;
+	const FName& DefaultPipelineStackName = bReimport ? ReimportPipelineName : UE::Interchange::Private::GetDefaultPipelineStackName(bSceneImport);
+	const TMap<FName, FInterchangePipelineStack>& PipelineStacks =
+		[this]() -> const TMap<FName, FInterchangePipelineStack>&
+		{
+			if (bSceneImport)
+			{
+				return GetDefault<UInterchangeProjectSettings>()->SceneImportSettings.PipelineStacks;
+			}
+			else
+			{
+				return GetDefault<UInterchangeProjectSettings>()->ContentImportSettings.PipelineStacks;
+			}
+		}();
+		
 
 	//In reimport we modify directly the asset pipeline
 	if (bReimport)
@@ -84,7 +117,7 @@ void SInterchangePipelineStacksTreeView::Construct(const FArguments& InArgs)
 	}
 	else
 	{
-		for (const TPair<FName, FInterchangePipelineStack>& NameAndPipelineStack : DefaultPipelineStacks)
+		for (const TPair<FName, FInterchangePipelineStack>& NameAndPipelineStack : PipelineStacks)
 		{
 			TSharedPtr<FInterchangePipelineStacksTreeNodeItem> StackNode = MakeShared<FInterchangePipelineStacksTreeNodeItem>();
 			StackNode->StackName = NameAndPipelineStack.Key;
@@ -131,6 +164,7 @@ public:
 
 		/** The item content. */
 		SLATE_ARGUMENT(TSharedPtr<FInterchangePipelineStacksTreeNodeItem>, InterchangeNode)
+		SLATE_ARGUMENT(bool, bSceneImport)
 	SLATE_END_ARGS()
 
 	/**
@@ -141,6 +175,8 @@ public:
 	void Construct(const FArguments& InArgs, const TSharedRef<STableViewBase>& InOwnerTableView)
 	{
 		InterchangeNode = InArgs._InterchangeNode;
+		bSceneImport = InArgs._bSceneImport;
+
 		//This is suppose to always be valid
 		check(InterchangeNode);
 		const bool bIsPipelineStackNode = InterchangeNode->Pipeline == nullptr;
@@ -185,7 +221,7 @@ public:
 private:
 	const FSlateBrush* GetImageItemIcon() const
 	{
-		const FName DefaultPipelineStackName = GetDefault<UInterchangeProjectSettings>()->DefaultPipelineStack;
+		const FName DefaultPipelineStackName = UE::Interchange::Private::GetDefaultPipelineStackName(bSceneImport);
 		//This is suppose to always be valid
 		check(InterchangeNode);
 		const bool bIsPipelineStackNode = InterchangeNode->Pipeline == nullptr;
@@ -207,12 +243,14 @@ private:
 
 	/** The node to build the tree view row from. */
 	TSharedPtr<FInterchangePipelineStacksTreeNodeItem> InterchangeNode = nullptr;
+	bool bSceneImport = false;
 };
 
 TSharedRef< ITableRow > SInterchangePipelineStacksTreeView::OnGenerateRowPipelineConfigurationTreeView(TSharedPtr<FInterchangePipelineStacksTreeNodeItem> Item, const TSharedRef< STableViewBase >& OwnerTable)
 {
 	TSharedRef< SInterchangePipelineStacksTreeViewItem > ReturnRow = SNew(SInterchangePipelineStacksTreeViewItem, OwnerTable)
-		.InterchangeNode(Item);
+		.InterchangeNode(Item)
+		.bSceneImport(bSceneImport);
 	return ReturnRow;
 }
 void SInterchangePipelineStacksTreeView::OnGetChildrenPipelineConfigurationTreeView(TSharedPtr<FInterchangePipelineStacksTreeNodeItem> InParent, TArray< TSharedPtr<FInterchangePipelineStacksTreeNodeItem> >& OutChildren)
@@ -287,7 +325,7 @@ TSharedPtr<SWidget> SInterchangePipelineStacksTreeView::OnOpenContextMenu()
 			MenuBuilder.BeginSection("TreeViewContextMenuStackNodeSection");
 			{
 				const FSlateIcon DefaultIcon(FAppStyle::GetAppStyleSetName(), "Icons.Default");
-				MenuBuilder.AddMenuEntry(LOCTEXT("SetHasDefaultMenuAction", "Set Has Default Stack"), FText(), DefaultIcon, FUIAction(FExecuteAction::CreateSP(this, &SInterchangePipelineStacksTreeView::SetHasDefaultStack, SelectNode->StackName)));
+				MenuBuilder.AddMenuEntry(LOCTEXT("SetAsDefaultMenuAction", "Set As Default Stack"), FText(), DefaultIcon, FUIAction(FExecuteAction::CreateSP(this, &SInterchangePipelineStacksTreeView::SetAsDefaultStack, SelectNode->StackName)));
 			}
 			MenuBuilder.EndSection();
 		}
@@ -297,11 +335,34 @@ TSharedPtr<SWidget> SInterchangePipelineStacksTreeView::OnOpenContextMenu()
 	return MenuBuilder.MakeWidget();
 }
 
-void SInterchangePipelineStacksTreeView::SetHasDefaultStack(FName NewDefaultStackValue)
+void SInterchangePipelineStacksTreeView::SetAsDefaultStack(FName NewDefaultStackValue)
 {
 	UInterchangeProjectSettings* InterchangeProjectSettingsCDO = GetMutableDefault<UInterchangeProjectSettings>();
-	FName& DefaultPipelineStackName = InterchangeProjectSettingsCDO->DefaultPipelineStack;
-	const TMap<FName, FInterchangePipelineStack>& PipelineStacks = InterchangeProjectSettingsCDO->PipelineStacks;
+	FName& DefaultPipelineStackName = [this, InterchangeProjectSettingsCDO]() -> FName&
+	{
+		if (bSceneImport)
+		{
+			return InterchangeProjectSettingsCDO->SceneImportSettings.DefaultPipelineStack;
+		}
+		else
+		{
+			return InterchangeProjectSettingsCDO->ContentImportSettings.DefaultPipelineStack;
+		}
+	}();
+
+	const TMap<FName, FInterchangePipelineStack>& PipelineStacks =
+		[this, InterchangeProjectSettingsCDO]() -> const TMap<FName, FInterchangePipelineStack>&
+		{
+			if (bSceneImport)
+			{
+				return InterchangeProjectSettingsCDO->SceneImportSettings.PipelineStacks;
+			}
+			else
+			{
+				return InterchangeProjectSettingsCDO->ContentImportSettings.PipelineStacks;
+			}
+		}();
+
 	if (PipelineStacks.Contains(NewDefaultStackValue))
 	{
 		DefaultPipelineStackName = NewDefaultStackValue;
@@ -347,6 +408,7 @@ TSharedRef<SBox> SInterchangePipelineConfigurationDialog::SpawnPipelineConfigura
 	//Create the treeview
 	PipelineConfigurationTreeView = SNew(SInterchangePipelineStacksTreeView)
 		.OnSelectionChangedDelegate(this, &SInterchangePipelineConfigurationDialog::OnSelectionChanged)
+		.bSceneImport(bSceneImport)
 		.bReimport(bReimport)
 		.PipelineStack(PipelineStack);
 
@@ -419,6 +481,7 @@ void SInterchangePipelineConfigurationDialog::Construct(const FArguments& InArgs
 
 	OwnerWindow = InArgs._OwnerWindow;
 	SourceData = InArgs._SourceData;
+	bSceneImport = InArgs._bSceneImport;
 	bReimport = InArgs._bReimport;
 	PipelineStack = InArgs._PipelineStack;
 
@@ -606,7 +669,7 @@ FReply SInterchangePipelineConfigurationDialog::OnResetToDefault()
 			else
 			{
 				//The reset to default when doing a import is to duplicate again the pipeline stack reference asset.
-				const TMap<FName, FInterchangePipelineStack>& DefaultPipelineStacks = GetDefault<UInterchangeProjectSettings>()->PipelineStacks;
+				const TMap<FName, FInterchangePipelineStack>& DefaultPipelineStacks = UE::Interchange::Private::GetDefaultImportSettings(bSceneImport).PipelineStacks;
 				//Find the stack node, so we can update the pipeline
 				TSharedPtr<FInterchangePipelineStacksTreeNodeItem> StackNode;
 				TArray<TSharedPtr<FInterchangePipelineStacksTreeNodeItem>>& RootNodeArray = PipelineConfigurationTreeView->GetMutableRootNodeArray();
@@ -670,11 +733,11 @@ bool SInterchangePipelineConfigurationDialog::RecursiveValidatePipelineSettings(
 
 bool SInterchangePipelineConfigurationDialog::IsImportButtonEnabled() const
 {
-	const FName DefaultPipelineStackName = GetDefault<UInterchangeProjectSettings>()->DefaultPipelineStack;
+	const FName DefaultPipelineStackName = UE::Interchange::Private::GetDefaultPipelineStackName(bSceneImport);
 	const TArray<TSharedPtr<FInterchangePipelineStacksTreeNodeItem>>& RootNodeArray = PipelineConfigurationTreeView->GetRootNodeArray();
 	for (const TSharedPtr<FInterchangePipelineStacksTreeNodeItem>& RootNode : RootNodeArray)
 	{
-		//Reimport have only one stack and the name wont match DefaultPipelineStackName
+		//Reimport has only one stack and the name wont match DefaultPipelineStackName
 		if (bReimport || RootNode->StackName == DefaultPipelineStackName)
 		{
 			TOptional<FText> InvalidReason;
@@ -686,7 +749,7 @@ bool SInterchangePipelineConfigurationDialog::IsImportButtonEnabled() const
 
 FText SInterchangePipelineConfigurationDialog::GetImportButtonTooltip() const
 {
-	const FName DefaultPipelineStackName = GetDefault<UInterchangeProjectSettings>()->DefaultPipelineStack;
+	const FName DefaultPipelineStackName = UE::Interchange::Private::GetDefaultPipelineStackName(bSceneImport);
 	const TArray<TSharedPtr<FInterchangePipelineStacksTreeNodeItem>>& RootNodeArray = PipelineConfigurationTreeView->GetRootNodeArray();
 	for (const TSharedPtr<FInterchangePipelineStacksTreeNodeItem>& RootNode : RootNodeArray)
 	{
diff --git a/Engine/Plugins/Experimental/Interchange/Editor/Source/Pipelines/Private/SInterchangePipelineConfigurationDialog.h b/Engine/Plugins/Experimental/Interchange/Editor/Source/Pipelines/Private/SInterchangePipelineConfigurationDialog.h
--- a/Engine/Plugins/Experimental/Interchange/Editor/Source/Pipelines/Private/SInterchangePipelineConfigurationDialog.h
+++ b/Engine/Plugins/Experimental/Interchange/Editor/Source/Pipelines/Private/SInterchangePipelineConfigurationDialog.h
@@ -53,6 +53,7 @@ public:
 		: _OnSelectionChangedDelegate()
 	{}
 		SLATE_EVENT(FOnPipelineConfigurationSelectionChanged, OnSelectionChangedDelegate)
+		SLATE_ARGUMENT(bool, bSceneImport)
 		SLATE_ARGUMENT(bool, bReimport)
 		SLATE_ARGUMENT(TArray<UInterchangePipelineBase*>, PipelineStack)
 	SLATE_END_ARGS()
@@ -70,6 +71,7 @@ public:
 protected:
 	/** Delegate to invoke when selection changes. */
 	FOnPipelineConfigurationSelectionChanged OnSelectionChangedDelegate;
+	bool bSceneImport = false;
 	bool bReimport = false;
 	TArray<UInterchangePipelineBase*> PipelineStack;
 
@@ -78,7 +80,7 @@ protected:
 
 	/** Open a context menu for the current selection */
 	TSharedPtr<SWidget> OnOpenContextMenu();
-	void SetHasDefaultStack(FName NewDefaultStackValue);
+	void SetAsDefaultStack(FName NewDefaultStackValue);
 	void RecursiveSetExpand(TSharedPtr<FInterchangePipelineStacksTreeNodeItem> Node, bool ExpandState);
 	void OnTreeViewSelectionChanged(TSharedPtr<FInterchangePipelineStacksTreeNodeItem> Item, ESelectInfo::Type SelectionType);
 };
@@ -100,6 +102,7 @@ public:
 
 		SLATE_ARGUMENT(TWeakPtr<SWindow>, OwnerWindow)
 		SLATE_ARGUMENT(TWeakObjectPtr<UInterchangeSourceData>, SourceData)
+		SLATE_ARGUMENT(bool, bSceneImport)
 		SLATE_ARGUMENT(bool, bReimport)
 		SLATE_ARGUMENT(TArray<UInterchangePipelineBase*>, PipelineStack)
 	SLATE_END_ARGS()
@@ -156,6 +159,7 @@ private:
 	TSharedPtr<IDetailsView> PipelineConfigurationDetailsView;
 	TSharedPtr<SCheckBox> UseSameSettingsForAllCheckBox;
 
+	bool bSceneImport = false;
 	bool bReimport = false;
 	bool bCanceled = false;
 	bool bImportAll = false;
diff --git a/Engine/Plugins/Experimental/Interchange/Editor/Source/Pipelines/Public/InterchangePipelineConfigurationGeneric.h b/Engine/Plugins/Experimental/Interchange/Editor/Source/Pipelines/Public/InterchangePipelineConfigurationGeneric.h
--- a/Engine/Plugins/Experimental/Interchange/Editor/Source/Pipelines/Public/InterchangePipelineConfigurationGeneric.h
+++ b/Engine/Plugins/Experimental/Interchange/Editor/Source/Pipelines/Public/InterchangePipelineConfigurationGeneric.h
@@ -21,5 +21,6 @@ public:
 protected:
 
 	virtual EInterchangePipelineConfigurationDialogResult ShowPipelineConfigurationDialog(TWeakObjectPtr<UInterchangeSourceData> SourceData) override;
+	virtual EInterchangePipelineConfigurationDialogResult ShowScenePipelineConfigurationDialog(TWeakObjectPtr<UInterchangeSourceData> SourceData) override;
 	virtual EInterchangePipelineConfigurationDialogResult ShowReimportPipelineConfigurationDialog(TArray<UInterchangePipelineBase*>& PipelineStack, TWeakObjectPtr<UInterchangeSourceData> SourceData) override;
 };
diff --git a/Engine/Source/Editor/UnrealEd/Private/FileHelpers.cpp b/Engine/Source/Editor/UnrealEd/Private/FileHelpers.cpp
--- a/Engine/Source/Editor/UnrealEd/Private/FileHelpers.cpp
+++ b/Engine/Source/Editor/UnrealEd/Private/FileHelpers.cpp
@@ -81,6 +81,7 @@
 #include "PackageSourceControlHelper.h"
 #include "ActorFolder.h"
 #include "InterchangeManager.h"
+#include "InterchangeProjectSettings.h"
 
 DEFINE_LOG_CATEGORY_STATIC(LogFileHelpers, Log, All);
 
@@ -1442,7 +1443,7 @@ void FEditorFileUtils::Import(const FString& InFilename)
 	const UEditorExperimentalSettings* EditorExperimentalSettings = GetDefault<UEditorExperimentalSettings>();
 
 	UE::Interchange::FScopedSourceData ScopedSourceData(InFilename);
-	const bool bImportThroughInterchange = EditorExperimentalSettings->bEnableInterchangeFramework && UInterchangeManager::GetInterchangeManager().CanTranslateSourceData(ScopedSourceData.GetSourceData());
+	const bool bImportThroughInterchange = (EditorExperimentalSettings->bEnableInterchangeFramework && GetDefault<UInterchangeProjectSettings>()->bUseInterchangeWhenImportingIntoLevel) && UInterchangeManager::GetInterchangeManager().CanTranslateSourceData(ScopedSourceData.GetSourceData());
 
 	USceneImportFactory* SceneFactory = nullptr;
 
diff --git a/Engine/Source/Runtime/Experimental/Interchange/Engine/Private/InterchangeManager.cpp b/Engine/Source/Runtime/Experimental/Interchange/Engine/Private/InterchangeManager.cpp
--- a/Engine/Source/Runtime/Experimental/Interchange/Engine/Private/InterchangeManager.cpp
+++ b/Engine/Source/Runtime/Experimental/Interchange/Engine/Private/InterchangeManager.cpp
@@ -38,7 +38,7 @@
 
 
 
-namespace InternalInterchangePrivate
+namespace UE::Interchange::Private
 {
 	const FLogCategoryBase* GetLogInterchangePtr()
 	{
@@ -48,6 +48,29 @@ namespace InternalInterchangePrivate
 		return &LogInterchangeEngine;
 #endif
 	}
+
+	const FInterchangeImportSettings& GetImportSettings(const UInterchangeProjectSettings& InterchangeProjectSettings, EImportType ImportType)
+	{
+		if (ImportType == EImportType::ImportType_Asset)
+		{
+			return InterchangeProjectSettings.ContentImportSettings;
+		}
+		else if (ImportType == EImportType::ImportType_Scene)
+		{
+			return InterchangeProjectSettings.SceneImportSettings;
+		}
+		else
+		{
+			ensure(false);
+			static FInterchangeImportSettings InvalidSettings;
+			return InvalidSettings;
+		}
+	}
+
+	FInterchangeImportSettings& GetMutableImportSettings(UInterchangeProjectSettings& InterchangeProjectSettings, EImportType ImportType)
+	{
+		return const_cast<FInterchangeImportSettings&>(GetImportSettings(InterchangeProjectSettings, ImportType));
+	}
 }
 
 UE::Interchange::FScopedSourceData::FScopedSourceData(const FString& Filename)
@@ -592,7 +615,7 @@ void UInterchangeManager::StartQueuedTasks(bool bCancelAllTasks /*= false*/)
 			NotificationConfig.bIsHeadless = false;
 			NotificationConfig.bKeepOpenOnFailure = true;
 			NotificationConfig.TitleText = TitleText;
-			NotificationConfig.LogCategory = InternalInterchangePrivate::GetLogInterchangePtr();
+			NotificationConfig.LogCategory = UE::Interchange::Private::GetLogInterchangePtr();
 			NotificationConfig.bCanCancel.Set(true);
 			NotificationConfig.bKeepOpenOnFailure.Set(true);
 
@@ -698,6 +721,34 @@ UInterchangeManager::ImportInternal(const FString& ContentPath, const UInterchan
 		//Import process can be started only in the game thread
 		return TTuple<UE::Interchange::FAssetImportResultRef, UE::Interchange::FSceneImportResultRef>{ MakeShared< UE::Interchange::FImportResult, ESPMode::ThreadSafe >(), MakeShared< UE::Interchange::FImportResult, ESPMode::ThreadSafe >() };
 	}
+
+	if (!SourceData)
+	{
+		UE_LOG(LogInterchangeEngine, Error, TEXT("Cannot import file, the source data is invalid"));
+		return TTuple<UE::Interchange::FAssetImportResultRef, UE::Interchange::FSceneImportResultRef>{ MakeShared< UE::Interchange::FImportResult, ESPMode::ThreadSafe >(), MakeShared< UE::Interchange::FImportResult, ESPMode::ThreadSafe >() };
+	}
+
+	const bool bImportScene = ImportType == UE::Interchange::EImportType::ImportType_Scene;
+	const UInterchangeProjectSettings* InterchangeProjectSettings = GetDefault<UInterchangeProjectSettings>();
+	const FInterchangeImportSettings& InterchangeImportSettings = UE::Interchange::Private::GetImportSettings(*InterchangeProjectSettings, ImportType);
+	
+	if (InterchangeImportSettings.PipelineStacks.Num() == 0)
+	{
+		UE_LOG(LogInterchangeEngine, Error, TEXT("Cannot import file, there is no pipeline stack define for %s import type. File path %s"), bImportScene ? TEXT("scene") : TEXT("content"));
+		return TTuple<UE::Interchange::FAssetImportResultRef, UE::Interchange::FSceneImportResultRef>{ MakeShared< UE::Interchange::FImportResult, ESPMode::ThreadSafe >(), MakeShared< UE::Interchange::FImportResult, ESPMode::ThreadSafe >() };
+	}
+	
+	//Set a default pipeline stack if none is valid
+	if (!InterchangeImportSettings.PipelineStacks.Contains(InterchangeImportSettings.DefaultPipelineStack))
+	{
+		UInterchangeProjectSettings* MutableInterchangeProjectSettings = GetMutableDefault<UInterchangeProjectSettings>();
+		FInterchangeImportSettings& MutableInterchangeImportSettings = UE::Interchange::Private::GetMutableImportSettings(*MutableInterchangeProjectSettings, ImportType);
+		TArray<FName> Keys;
+		MutableInterchangeImportSettings.PipelineStacks.GetKeys(Keys);
+		check(Keys.Num() > 0);
+		MutableInterchangeImportSettings.DefaultPipelineStack = Keys[0];
+	}
+
 	UInterchangeAssetImportData* OriginalAssetImportData = nullptr;
 	FString PackageBasePath = ContentPath;
 	if(!ImportAssetParameters.ReimportAsset)
@@ -745,12 +796,13 @@ UInterchangeManager::ImportInternal(const FString& ContentPath, const UInterchan
 		check(AsyncHelper->BaseNodeContainers[SourceDataIndex].IsValid());
 	}
 
-	const UInterchangeProjectSettings* InterchangeProjectSettings = GetDefault<UInterchangeProjectSettings>();
+	
 	UInterchangePipelineConfigurationBase* RegisteredPipelineConfiguration = nullptr;
 
 	//In runtime we do not have any pipeline configurator
 #if WITH_EDITORONLY_DATA
-	TSoftClassPtr <UInterchangePipelineConfigurationBase> PipelineConfigurationDialogClass = InterchangeProjectSettings->PipelineConfigurationDialogClass;
+	TSoftClassPtr <UInterchangePipelineConfigurationBase> PipelineConfigurationDialogClass = InterchangeImportSettings.PipelineConfigurationDialogClass;
+
 	if (PipelineConfigurationDialogClass.IsValid())
 	{
 		UClass* PipelineConfigurationClass = PipelineConfigurationDialogClass.LoadSynchronous();
@@ -768,27 +820,13 @@ UInterchangeManager::ImportInternal(const FString& ContentPath, const UInterchan
 
 #if WITH_EDITORONLY_DATA
 		const bool bShowPipelineStacksConfigurationDialog = !bIsUnattended
-															&& InterchangeProjectSettings->bShowPipelineStacksConfigurationDialog
+															&& InterchangeImportSettings.bShowPipelineStacksConfigurationDialog
 															&& !bImportAllWithDefault;
 #else
 		const bool bShowPipelineStacksConfigurationDialog = false;
 #endif
 															
-
-
-		auto GetDefaultPipelineStackName = [ImportType, &InterchangeProjectSettings]()->FName
-		{
-			if (ImportType == UE::Interchange::EImportType::ImportType_Scene)
-			{
-				return InterchangeProjectSettings->DefaultScenePipelineStack;
-			}
-			else
-			{
-				return InterchangeProjectSettings->DefaultPipelineStack;
-			}
-		};
-			
-		const TMap<FName, FInterchangePipelineStack>& DefaultPipelineStacks = InterchangeProjectSettings->PipelineStacks;
+		const TMap<FName, FInterchangePipelineStack>& DefaultPipelineStacks = InterchangeImportSettings.PipelineStacks;
 
 		//If we reimport we want to load the original pipeline and the original pipeline settings
 		if (OriginalAssetImportData && OriginalAssetImportData->Pipelines.Num() > 0)
@@ -861,11 +899,11 @@ UInterchangeManager::ImportInternal(const FString& ContentPath, const UInterchan
 		}
 		else
 		{
-			FName PipelineStackName = GetDefaultPipelineStackName();
+			FName PipelineStackName = InterchangeImportSettings.DefaultPipelineStack;
 			if (RegisteredPipelineConfiguration && (bShowPipelineStacksConfigurationDialog || (!DefaultPipelineStacks.Contains(PipelineStackName) && !bIsUnattended)))
 			{
 				//Show the dialog, a plugin should have register this dialog. We use a plugin to be able to use editor code when doing UI
-				EInterchangePipelineConfigurationDialogResult DialogResult = RegisteredPipelineConfiguration->ScriptedShowPipelineConfigurationDialog(DuplicateSourceData);
+				EInterchangePipelineConfigurationDialogResult DialogResult = bImportScene ? RegisteredPipelineConfiguration->ScriptedShowScenePipelineConfigurationDialog(DuplicateSourceData) : RegisteredPipelineConfiguration->ScriptedShowPipelineConfigurationDialog(DuplicateSourceData);
 				if (DialogResult == EInterchangePipelineConfigurationDialogResult::Cancel)
 				{
 					bImportCancel = true;
@@ -874,15 +912,15 @@ UInterchangeManager::ImportInternal(const FString& ContentPath, const UInterchan
 				{
 					bImportAllWithDefault = true;
 				}
-				PipelineStackName = GetDefaultPipelineStackName();
+				PipelineStackName = InterchangeImportSettings.DefaultPipelineStack;
 			}
 			if (!bImportCancel)
 			{
 				if (!DefaultPipelineStacks.Contains(PipelineStackName))
 				{
-					if (DefaultPipelineStacks.Contains(GetDefaultPipelineStackName()))
+					if (DefaultPipelineStacks.Contains(InterchangeImportSettings.DefaultPipelineStack))
 					{
-						PipelineStackName = GetDefaultPipelineStackName();
+						PipelineStackName = InterchangeImportSettings.DefaultPipelineStack;
 					}
 					else
 					{
@@ -944,6 +982,7 @@ UInterchangeManager::ImportInternal(const FString& ContentPath, const UInterchan
 			UInterchangePipelineBase* GeneratedPipeline = DuplicateObject<UInterchangePipelineBase>(ImportAssetParameters.OverridePipelines[GraphPipelineIndex], GetTransientPackage());
 			GeneratedPipeline->bAllowLockedPropertiesEdition = false;
 			AsyncHelper->Pipelines.Add(GeneratedPipeline);
+			AsyncHelper->OriginalPipelines.Add(GeneratedPipeline);
 		}
 	}
 
diff --git a/Engine/Source/Runtime/Experimental/Interchange/Engine/Public/InterchangePipelineConfigurationBase.h b/Engine/Source/Runtime/Experimental/Interchange/Engine/Public/InterchangePipelineConfigurationBase.h
--- a/Engine/Source/Runtime/Experimental/Interchange/Engine/Public/InterchangePipelineConfigurationBase.h
+++ b/Engine/Source/Runtime/Experimental/Interchange/Engine/Public/InterchangePipelineConfigurationBase.h
@@ -37,6 +37,19 @@ public:
 		return ShowPipelineConfigurationDialog(SourceData);
 	}
 
+	/**
+	 * Non virtual helper to allow blueprint to implement event base function to implement the ShowPipelineConfigurationDialog,
+	 */
+	UFUNCTION(BlueprintNativeEvent, BlueprintCallable, Category = "Interchange | Translator")
+	EInterchangePipelineConfigurationDialogResult ScriptedShowScenePipelineConfigurationDialog(UInterchangeSourceData* SourceData);
+
+	/** The default implementation (call if the blueprint do not have any implementation) will call the virtual ExecuteImportPipeline */
+	EInterchangePipelineConfigurationDialogResult ScriptedShowScenePipelineConfigurationDialog_Implementation(UInterchangeSourceData* SourceData)
+	{
+		//By default we call the virtual import pipeline execution
+		return ShowScenePipelineConfigurationDialog(SourceData);
+	}
+
 	/**
 	 * Non virtual helper to allow blueprint to implement event base function to implement the ShowPipelineConfigurationDialog,
 	 */
@@ -61,6 +74,15 @@ protected:
 		return EInterchangePipelineConfigurationDialogResult::Cancel;
 	}
 
+	/**
+	 * This function show a dialog use to configure pipeline stacks and return a stack name that tell the caller the user choice.
+	 */
+	virtual EInterchangePipelineConfigurationDialogResult ShowScenePipelineConfigurationDialog(TWeakObjectPtr<UInterchangeSourceData> SourceData)
+	{ 
+		//Not implemented
+		return EInterchangePipelineConfigurationDialogResult::Cancel;
+	}
+
 	/**
 	 * This function show a dialog use to configure pipeline stacks and return a stack name that tell the caller the user choice.
 	 */
diff --git a/Engine/Source/Runtime/Experimental/Interchange/Engine/Public/InterchangeProjectSettings.h b/Engine/Source/Runtime/Experimental/Interchange/Engine/Public/InterchangeProjectSettings.h
--- a/Engine/Source/Runtime/Experimental/Interchange/Engine/Public/InterchangeProjectSettings.h
+++ b/Engine/Source/Runtime/Experimental/Interchange/Engine/Public/InterchangeProjectSettings.h
@@ -17,52 +17,63 @@ struct FInterchangePipelineStack
 {
 	GENERATED_BODY()
 	
-	/** The starting mesh for the blueprint **/
-	UPROPERTY(EditAnywhere, Category = Interchange, meta=(AllowedClasses = "/Script/InterchangeCore.InterchangePipelineBase, /Script/InterchangeEngine.InterchangeBlueprintPipelineBase, /Script/InterchangeEngine.InterchangePythonPipelineAsset"))
+	UPROPERTY(EditAnywhere, Category = "Pipelines", meta = (AllowedClasses = "/Script/InterchangeCore.InterchangePipelineBase, /Script/InterchangeEngine.InterchangeBlueprintPipelineBase, /Script/InterchangeEngine.InterchangePythonPipelineAsset"))
 	TArray<FSoftObjectPath> Pipelines;
 };
 
-UCLASS(config=Engine, meta=(DisplayName=Interchange), MinimalAPI)
-class UInterchangeProjectSettings : public UDeveloperSettings
+USTRUCT()
+struct FInterchangeImportSettings
 {
 	GENERATED_BODY()
 
-public:
 	/** All the available pipeline stacks you want to use to import with interchange. The chosen pipeline stack execute all the pipelines from top to bottom order. You can order them by using the grip on the left of any pipelines.*/
-	UPROPERTY(EditAnywhere, config, Category = Interchange)
+	UPROPERTY(EditAnywhere, Category = "Pipeline")
 	TMap<FName, FInterchangePipelineStack> PipelineStacks;
 
-	/** This tell interchange which pipeline to select when importing assets.*/
-	UPROPERTY(EditAnywhere, config, Category = Interchange)
+	/** This tell interchange which pipeline stack to select when importing.*/
+	UPROPERTY(EditAnywhere, Category = "Pipeline")
 	FName DefaultPipelineStack;
 
-#if WITH_EDITORONLY_DATA
 	/** This tell interchange which pipeline configuration dialog to popup when we need to configure the pipelines.*/
-	UPROPERTY(EditAnywhere, config, Category = Interchange)
+	UPROPERTY(EditAnywhere, Category = "Pipeline")
 	TSoftClassPtr <UInterchangePipelineConfigurationBase> PipelineConfigurationDialogClass;
 
-	/** If enabled, the pipeline stacks configuration dialog will show every time interchange must choose a pipeline to import or re-import. If disabled interchange will use the DefaultPipelineStack.*/
-	UPROPERTY(EditAnywhere, config, Category = Interchange)
+	/** If enabled, the pipeline stacks configuration dialog will show when interchange must choose a pipeline to import or re-import. If disabled interchange will use the DefaultPipelineStack.*/
+	UPROPERTY(EditAnywhere, Category = "Pipeline")
 	bool bShowPipelineStacksConfigurationDialog;
+};
 
-	/** This tells interchange which file picker class to construct when we need to choose a file for a source.*/
-	UPROPERTY(EditAnywhere, config, Category = Interchange)
-	TSoftClassPtr <UInterchangeFilePickerBase> FilePickerClass;
-#endif
+UCLASS(config=Engine, meta=(DisplayName=Interchange), MinimalAPI)
+class UInterchangeProjectSettings : public UDeveloperSettings
+{
+	GENERATED_BODY()
+
+public:
+	/**
+	 * Settings used when importing into the content browser.
+	 */
+	UPROPERTY(EditAnywhere, config, Category = "ImportContent")
+	FInterchangeImportSettings ContentImportSettings;
 
 	/** If checked, will use Interchange when importing into level.*/
-	UPROPERTY(EditAnywhere, config, Category = "Interchange (Experimental)")
+	UPROPERTY(EditAnywhere, config, Category = "ImportIntoLevel|Experimental")
 	bool bUseInterchangeWhenImportingIntoLevel;
 
-	/** This tell interchange which pipeline to select when importing scenes.*/
-	UPROPERTY(EditAnywhere, config, Category = "Interchange (Experimental)")
-	FName DefaultScenePipelineStack;
+	/**
+	 * Settings used when importing into a level.
+	¸*/
+	UPROPERTY(EditAnywhere, config, Category = "ImportIntoLevel")
+	FInterchangeImportSettings SceneImportSettings;
+
+	/** This tells interchange which file picker class to construct when we need to choose a file for a source.*/
+	UPROPERTY(EditAnywhere, config, Category = "EditorInterface")
+	TSoftClassPtr <UInterchangeFilePickerBase> FilePickerClass;
 
 	/**
 	 * If checked, interchange translators and legacy importer will default static mesh geometry to smooth edge when the smoothing information is missing.
 	 * This option exist to allows old project to import the same way as before if their workflows need static mesh edges to be hard when the smoothing
 	 * info is missing.
 	 */
-	UPROPERTY(EditAnywhere, config, Category = "Import")
+	UPROPERTY(EditAnywhere, config, Category = "Generic|ImportSettings")
 	bool bStaticMeshUseSmoothEdgesIfSmoothingInformationIsMissing = true;
 };
\ No newline at end of file
