Interchange: Allow visibility edition for pipeline properties, also did a refactor to allow pipeline to adjust their settings from the context (import asset, re-import asset, import scene, ...)
#rnx

diff --git a/Engine/Plugins/Experimental/Interchange/Editor/Source/Pipelines/Private/InterchangeEditorPipelineDetails.cpp b/Engine/Plugins/Experimental/Interchange/Editor/Source/Pipelines/Private/InterchangeEditorPipelineDetails.cpp
--- a/Engine/Plugins/Experimental/Interchange/Editor/Source/Pipelines/Private/InterchangeEditorPipelineDetails.cpp
+++ b/Engine/Plugins/Experimental/Interchange/Editor/Source/Pipelines/Private/InterchangeEditorPipelineDetails.cpp
@@ -130,7 +130,8 @@ void FInterchangePipelineBaseDetailsCustomization::CustomizeDetails(IDetailLayou
 		return;
 	}
 
-	const bool bAllowLockedPropertiesEdition = InterchangePipeline->bAllowLockedPropertiesEdition;
+	const bool bAllowPropertyStatesEdition = InterchangePipeline->CanEditPropertiesStates();
+	const bool bIsReimportContext = InterchangePipeline->IsReimportContext();
 
 	TArray<FName> AllCategoryNames;
 	CachedDetailBuilder->GetCategoryNames(AllCategoryNames);
@@ -156,37 +157,43 @@ void FInterchangePipelineBaseDetailsCustomization::CustomizeDetails(IDetailLayou
 				continue;
 			}
 			const FName PropertyName = PropertyHandle.Get().GetProperty() ? PropertyHandle.Get().GetProperty()->GetFName() : NAME_None;
-			if (PropertyName == UInterchangePipelineBase::GetLockedPropertiesPropertyName())
+			if (PropertyName == UInterchangePipelineBase::GetPropertiesStatesPropertyName())
 			{
 				CachedDetailBuilder->HideProperty(PropertyHandle);
 				continue;
 			}
-			// Skip the property if:
-			// We do not allow property edition
-			// property is not in the List of supported properties
+			// Skip the property not in the List of supported properties
 			if (!CategoryAndProperties.Value.Contains(PropertyName))
 			{
 				continue;
 			}
 			FName PropertyPath = FName(PropertyPtr->GetPathName());
 			CachedDetailBuilder->HideProperty(PropertyHandle);
-			if (!bAllowLockedPropertiesEdition)
+			if (!bAllowPropertyStatesEdition)
 			{
-				const bool bLockedProperty = InterchangePipeline->GetLockedPropertyStatus(PropertyPath);
 				const FName SubCategoryData = FName(PropertyHandle->GetMetaData(TEXT("SubCategory")));
+				bool IsLocked = false;
+				if (const FInterchangePipelinePropertyStates* PropertyStates = InterchangePipeline->GetPropertyStates(PropertyPath))
+				{
+					if (!PropertyStates->IsPropertyVisible(bIsReimportContext))
+					{
+						continue;
+					}
+					IsLocked = PropertyStates->IsPropertyLocked();
+				}
 				if (SubCategoryData != NAME_None)
 				{
 					TMap<FName, TArray<FInternalPropertyData>>& SubCategoriesProperties = SubCategoriesPropertiesPerMainCategory.FindOrAdd(CategoryName);
 					TArray<FInternalPropertyData>& SubCategoryProperties = SubCategoriesProperties.FindOrAdd(SubCategoryData);
 					FInternalPropertyData& PropertyData = SubCategoryProperties.AddDefaulted_GetRef();
 					PropertyData.PropertyHandle = PropertyHandle;
-					PropertyData.bReadOnly = bLockedProperty;
+					PropertyData.bReadOnly = IsLocked;
 				}
 				else
 				{
 					IDetailPropertyRow& PropertyRow = Category.AddProperty(PropertyHandle);
 					//When we use the pipeline in interchange
-					if (bLockedProperty)
+					if (IsLocked)
 					{
 						LockPropertyHandleRow(PropertyHandle, PropertyRow);
 					}
@@ -199,12 +206,28 @@ void FInterchangePipelineBaseDetailsCustomization::CustomizeDetails(IDetailLayou
 				TSharedPtr<SWidget> NameWidget;
 				TSharedPtr<SWidget> ValueWidget;
 				FDetailWidgetRow Row;
+				auto TextColorLambda = [InterchangePipelinePtr = InterchangePipeline, PropertyPath]()
+				{
+					FSlateColor SlateColor = FStyleColors::Foreground;
+					if (InterchangePipelinePtr.IsValid())
+					{
+						if (const FInterchangePipelinePropertyStates* PropertyStates = InterchangePipelinePtr->GetPropertyStates(PropertyPath))
+						{
+							if (PropertyStates->IsPropertyLocked())
+							{
+								SlateColor = FStyleColors::AccentBlue;
+							}
+						}
+					}
+					return SlateColor;
+				};
 				PropertyRow.GetDefaultWidgets(NameWidget, ValueWidget, Row);
 				//When we edit the default value and the locked properties of a pipeline assets
 				PropertyRow.CustomWidget(bShowChildren)
 				.NameContent()
 				.MinDesiredWidth(Row.NameWidget.MinWidth)
 				.MaxDesiredWidth(Row.NameWidget.MaxWidth)
+				.HAlign(HAlign_Fill)
 				[
 					SNew(SHorizontalBox)
 					+ SHorizontalBox::Slot()
@@ -225,9 +248,9 @@ void FInterchangePipelineBaseDetailsCustomization::CustomizeDetails(IDetailLayou
 							{
 								return;
 							}
-							FScopedTransaction ScopedTransaction(NSLOCTEXT("InterchangePipelineBaseDetails::CustomizeDetails", "TransactionLockedPropertiesSet", "Toggle Property Locked"), !GIsTransacting);
+							FScopedTransaction ScopedTransaction(NSLOCTEXT("InterchangePipelineBaseDetails::CustomizeDetails", "TransactionLockedPropertiesSet", "Toggle property locked at import"), !GIsTransacting);
 							InterchangePipelinePtr->Modify();
-							InterchangePipelinePtr->SetLockedPropertyStatus(PropertyPath, (CheckType == ECheckBoxState::Checked));
+							InterchangePipelinePtr->FindOrAddPropertyStates(PropertyPath).SetPropertyLocked(CheckType == ECheckBoxState::Checked);
 							InterchangePipelinePtr->PostEditChange();
 						})
 						.IsChecked_Lambda([InterchangePipelinePtr = InterchangePipeline, PropertyPath]()
@@ -236,31 +259,119 @@ void FInterchangePipelineBaseDetailsCustomization::CustomizeDetails(IDetailLayou
 							{
 								return ECheckBoxState::Unchecked;
 							}
-							return InterchangePipelinePtr->GetLockedPropertyStatus(PropertyPath) ? ECheckBoxState::Checked : ECheckBoxState::Unchecked;
+							if (const FInterchangePipelinePropertyStates* PropertyStates = InterchangePipelinePtr->GetPropertyStates(PropertyPath))
+							{
+								return PropertyStates->IsPropertyLocked() ? ECheckBoxState::Checked : ECheckBoxState::Unchecked;
+							}
+							return ECheckBoxState::Unchecked;
 						})
 					]
 					+ SHorizontalBox::Slot()
 					.FillWidth(1.0f)
+					.HAlign(HAlign_Left)
 					.Padding(3.0f, 1.0f)
 					[
 						SNew(SBorder)
-						.BorderBackgroundColor(FStyleColors::Background)
-						.ForegroundColor_Lambda([InterchangePipelinePtr = InterchangePipeline, PropertyPath]()
-							{
-								if (!InterchangePipelinePtr.IsValid())
-								{
-									return FStyleColors::Foreground;
-								}
-								if (InterchangePipelinePtr->GetLockedPropertyStatus(PropertyPath))
-								{
-									return FStyleColors::AccentBlue;
-								}
-								return FStyleColors::Foreground;
-							})
+						.BorderBackgroundColor(FStyleColors::Transparent)
+						.ForegroundColor_Lambda(TextColorLambda)
 						[
 							NameWidget.ToSharedRef()
 						]
 					]
+					+ SHorizontalBox::Slot()
+					.AutoWidth()
+					.Padding(6.0f, 1.0f, 3.0f, 1.0f)
+					.VAlign(VAlign_Center)
+					[
+						SNew(STextBlock)
+						.Font(IDetailLayoutBuilder::GetDetailFont())
+						.Text(NSLOCTEXT("InterchangePipelineBaseDetails::CustomizeDetails", "HiddenAtImportText", "Import"))
+					]
+					+ SHorizontalBox::Slot()
+					.AutoWidth()
+					.Padding(3.0f, 1.0f)
+					[
+						SNew(SCheckBox)
+						.CheckedImage(FAppStyle::Get().GetBrush("Icons.Hidden"))
+						.CheckedHoveredImage(FAppStyle::Get().GetBrush("Icons.Hidden"))
+						.CheckedPressedImage(FAppStyle::Get().GetBrush("Icons.Hidden"))
+						.UncheckedImage(FAppStyle::Get().GetBrush("Icons.Visible"))
+						.UncheckedHoveredImage(FAppStyle::Get().GetBrush("Icons.Visible"))
+						.UncheckedPressedImage(FAppStyle::Get().GetBrush("Icons.Visible"))
+						.ToolTipText(NSLOCTEXT("InterchangePipelineBaseDetails::CustomizeDetails", "VisibleTooltipImport", "If true this property will be visible when displaying the interchange import dialog."))
+						.OnCheckStateChanged_Lambda([InterchangePipelinePtr = InterchangePipeline, PropertyPath](ECheckBoxState CheckType)
+						{
+							constexpr bool bIsReimportContextLocal = false;
+							if (!ensure(InterchangePipelinePtr.IsValid()))
+							{
+								return;
+							}
+							FScopedTransaction ScopedTransaction(NSLOCTEXT("InterchangePipelineBaseDetails::CustomizeDetails", "TransactionvisibilityPropertiesToggleImport", "Toggle property visibility at import"), !GIsTransacting);
+							InterchangePipelinePtr->Modify();
+							InterchangePipelinePtr->FindOrAddPropertyStates(PropertyPath).SetPropertyVisible(bIsReimportContextLocal, (CheckType != ECheckBoxState::Checked));
+							InterchangePipelinePtr->PostEditChange();
+						})
+						.IsChecked_Lambda([InterchangePipelinePtr = InterchangePipeline, PropertyPath]()
+						{
+							constexpr bool bIsReimportContextLocal = false;
+							if (!InterchangePipelinePtr.IsValid())
+							{
+								return ECheckBoxState::Unchecked;
+							}
+							if (const FInterchangePipelinePropertyStates* PropertyStates = InterchangePipelinePtr->GetPropertyStates(PropertyPath))
+							{
+								return PropertyStates->IsPropertyVisible(bIsReimportContextLocal) ? ECheckBoxState::Unchecked : ECheckBoxState::Checked;
+							}
+							return ECheckBoxState::Unchecked;
+						})
+					]
+					+ SHorizontalBox::Slot()
+					.AutoWidth()
+					.Padding(6.0f, 1.0f, 3.0f, 1.0f)
+					.VAlign(VAlign_Center)
+					[
+						SNew(STextBlock)
+						.Font(IDetailLayoutBuilder::GetDetailFont())
+						.Text(NSLOCTEXT("InterchangePipelineBaseDetails::CustomizeDetails", "HiddenAtReimportText", "Reimport"))
+					]
+					+ SHorizontalBox::Slot()
+					.AutoWidth()
+					.Padding(3.0f, 1.0f)
+					[
+						SNew(SCheckBox)
+						.CheckedImage(FAppStyle::Get().GetBrush("Icons.Hidden"))
+						.CheckedHoveredImage(FAppStyle::Get().GetBrush("Icons.Hidden"))
+						.CheckedPressedImage(FAppStyle::Get().GetBrush("Icons.Hidden"))
+						.UncheckedImage(FAppStyle::Get().GetBrush("Icons.Visible"))
+						.UncheckedHoveredImage(FAppStyle::Get().GetBrush("Icons.Visible"))
+						.UncheckedPressedImage(FAppStyle::Get().GetBrush("Icons.Visible"))
+						.ToolTipText(NSLOCTEXT("InterchangePipelineBaseDetails::CustomizeDetails", "VisibleTooltipReimport", "If true this property will be visible when displaying the interchange reimport dialog."))
+						.OnCheckStateChanged_Lambda([InterchangePipelinePtr = InterchangePipeline, PropertyPath](ECheckBoxState CheckType)
+						{
+							constexpr bool bIsReimportContextLocal = true;
+							if (!ensure(InterchangePipelinePtr.IsValid()))
+							{
+								return;
+							}
+							FScopedTransaction ScopedTransaction(NSLOCTEXT("InterchangePipelineBaseDetails::CustomizeDetails", "TransactionvisibilityPropertiesToggleReimport", "Toggle property visibility at reimport"), !GIsTransacting);
+							InterchangePipelinePtr->Modify();
+							InterchangePipelinePtr->FindOrAddPropertyStates(PropertyPath).SetPropertyVisible(bIsReimportContextLocal, (CheckType != ECheckBoxState::Checked));
+							InterchangePipelinePtr->PostEditChange();
+						})
+						.IsChecked_Lambda([InterchangePipelinePtr = InterchangePipeline, PropertyPath]()
+						{
+							constexpr bool bIsReimportContextLocal = true;
+							if (!InterchangePipelinePtr.IsValid())
+							{
+								return ECheckBoxState::Unchecked;
+							}
+							if (const FInterchangePipelinePropertyStates* PropertyStates = InterchangePipelinePtr->GetPropertyStates(PropertyPath))
+							{
+								return PropertyStates->IsPropertyVisible(bIsReimportContextLocal) ? ECheckBoxState::Unchecked : ECheckBoxState::Checked;
+							}
+							return ECheckBoxState::Unchecked;
+						})
+					]
 				]
 				.ValueContent()
 				.MinDesiredWidth(Row.ValueWidget.MinWidth)
diff --git a/Engine/Plugins/Experimental/Interchange/Editor/Source/Pipelines/Private/SInterchangePipelineConfigurationDialog.cpp b/Engine/Plugins/Experimental/Interchange/Editor/Source/Pipelines/Private/SInterchangePipelineConfigurationDialog.cpp
--- a/Engine/Plugins/Experimental/Interchange/Editor/Source/Pipelines/Private/SInterchangePipelineConfigurationDialog.cpp
+++ b/Engine/Plugins/Experimental/Interchange/Editor/Source/Pipelines/Private/SInterchangePipelineConfigurationDialog.cpp
@@ -105,7 +105,6 @@ void SInterchangePipelineStacksTreeView::Construct(const FArguments& InArgs)
 			{
 				//When we re-import we save the setting so ResetToDefault can simply just reload the settings
 				GeneratedPipeline->SaveSettings(ReimportPipelineName);
-				GeneratedPipeline->bAllowLockedPropertiesEdition = false;
 				TSharedPtr<FInterchangePipelineStacksTreeNodeItem> PipelineNode = MakeShared<FInterchangePipelineStacksTreeNodeItem>();
 				PipelineNode->StackName = ReimportPipelineName;
 				//Currently the re-import never lock properties
@@ -127,7 +126,7 @@ void SInterchangePipelineStacksTreeView::Construct(const FArguments& InArgs)
 			{
 				if (UInterchangePipelineBase* GeneratedPipeline = UE::Interchange::GeneratePipelineInstance(InterchangePipelineStack.Pipelines[PipelineIndex]))
 				{
-					GeneratedPipeline->bAllowLockedPropertiesEdition = false;
+					GeneratedPipeline->AdjustSettingsForContext((bSceneImport ? EInterchangePipelineContext::SceneImport : EInterchangePipelineContext::AssetImport), nullptr);
 					//Load the settings for this pipeline
 					GeneratedPipeline->LoadSettings(NameAndPipelineStack.Key);
 					GeneratedPipeline->PreDialogCleanup(NameAndPipelineStack.Key);
@@ -676,35 +675,35 @@ FReply SInterchangePipelineConfigurationDialog::OnResetToDefault()
 				for (TSharedPtr<FInterchangePipelineStacksTreeNodeItem>& RootNode : RootNodeArray)
 				{
 					RecursiveIterateNode(RootNode, [this, &Pipeline, &DefaultPipelineStacks](TSharedPtr<FInterchangePipelineStacksTreeNodeItem>& NodeItem)
+					{
+						if (NodeItem->StackName == CurrentStackName && NodeItem->Pipeline == Pipeline)
 						{
-							if (NodeItem->StackName == CurrentStackName && NodeItem->Pipeline == Pipeline)
+							//We assume the pipelines inside one stack are all different classes, we use the class to know which default asset we need to duplicate
+							const UClass* PipelineClass = NodeItem->Pipeline->GetClass();
+							for (const TPair<FName, FInterchangePipelineStack>& NameAndPipelineStack : DefaultPipelineStacks)
 							{
-								//We assume the pipelines inside one stack are all different classes, we use the class to know which default asset we need to duplicate
-								const UClass* PipelineClass = NodeItem->Pipeline->GetClass();
-								for (const TPair<FName, FInterchangePipelineStack>& NameAndPipelineStack : DefaultPipelineStacks)
+								if(CurrentStackName != NameAndPipelineStack.Key)
 								{
-									if(CurrentStackName != NameAndPipelineStack.Key)
-									{
-										continue;
-									}
-									const FInterchangePipelineStack& InterchangePipelineStack = NameAndPipelineStack.Value;
-									for (int32 PipelineIndex = 0; PipelineIndex < InterchangePipelineStack.Pipelines.Num(); ++PipelineIndex)
+									continue;
+								}
+								const FInterchangePipelineStack& InterchangePipelineStack = NameAndPipelineStack.Value;
+								for (int32 PipelineIndex = 0; PipelineIndex < InterchangePipelineStack.Pipelines.Num(); ++PipelineIndex)
+								{
+									if (UInterchangePipelineBase* GeneratedPipeline = UE::Interchange::GeneratePipelineInstance(InterchangePipelineStack.Pipelines[PipelineIndex]))
 									{
-										if (UInterchangePipelineBase* GeneratedPipeline = UE::Interchange::GeneratePipelineInstance(InterchangePipelineStack.Pipelines[PipelineIndex]))
+										if (GeneratedPipeline->GetClass() == PipelineClass)
 										{
-											if (GeneratedPipeline->GetClass() == PipelineClass)
-											{
-												NodeItem->Pipeline = GeneratedPipeline;
-												NodeItem->Pipeline->bAllowLockedPropertiesEdition = false;
-												PipelineConfigurationDetailsView->SetObject(NodeItem->Pipeline, true);
-												//Exit the lambda
-												return;
-											}
+											NodeItem->Pipeline = GeneratedPipeline;
+											GeneratedPipeline->AdjustSettingsForContext((bSceneImport ? EInterchangePipelineContext::SceneImport : EInterchangePipelineContext::AssetImport), nullptr);
+											PipelineConfigurationDetailsView->SetObject(NodeItem->Pipeline, true);
+											//Exit the lambda
+											return;
 										}
 									}
 								}
 							}
-						});
+						}
+					});
 				}
 			}
 		}
diff --git a/Engine/Plugins/Experimental/Interchange/Runtime/Source/Pipelines/Private/InterchangeGenericAnimationPipeline.cpp b/Engine/Plugins/Experimental/Interchange/Runtime/Source/Pipelines/Private/InterchangeGenericAnimationPipeline.cpp
--- a/Engine/Plugins/Experimental/Interchange/Runtime/Source/Pipelines/Private/InterchangeGenericAnimationPipeline.cpp
+++ b/Engine/Plugins/Experimental/Interchange/Runtime/Source/Pipelines/Private/InterchangeGenericAnimationPipeline.cpp
@@ -47,18 +47,23 @@ namespace UE::Interchange::Private
 	}
 }
 
-void UInterchangeGenericAnimationPipeline::AdjustSettingsForReimportType(EInterchangeReimportType ImportType, TObjectPtr<UObject> ReimportAsset)
+void UInterchangeGenericAnimationPipeline::AdjustSettingsForContext(EInterchangePipelineContext ImportType, TObjectPtr<UObject> ReimportAsset)
 {
+	Super::AdjustSettingsForContext(ImportType, ReimportAsset);
+
 	check(!CommonSkeletalMeshesAndAnimationsProperties.IsNull());
-	if (ImportType == EInterchangeReimportType::AssetCustomLODImport
-		|| ImportType == EInterchangeReimportType::AssetCustomLODReimport
-		|| ImportType == EInterchangeReimportType::AssetAlternateSkinningImport
-		|| ImportType == EInterchangeReimportType::AssetAlternateSkinningReimport)
+	
+	if (ImportType == EInterchangePipelineContext::AssetCustomLODImport
+		|| ImportType == EInterchangePipelineContext::AssetCustomLODReimport
+		|| ImportType == EInterchangePipelineContext::AssetAlternateSkinningImport
+		|| ImportType == EInterchangePipelineContext::AssetAlternateSkinningReimport)
 	{
 		bImportAnimations = false;
 		CommonSkeletalMeshesAndAnimationsProperties->bImportOnlyAnimations = false;
 	}
-	else if(ImportType == EInterchangeReimportType::AssetReimport)
+	
+	TArray<FString> HideCategories;
+	if (ImportType == EInterchangePipelineContext::AssetReimport)
 	{
 		if (UAnimSequence* AnimSequence = Cast<UAnimSequence>(ReimportAsset))
 		{
@@ -66,6 +71,18 @@ void UInterchangeGenericAnimationPipeline::AdjustSettingsForReimportType(EInterc
 			CommonSkeletalMeshesAndAnimationsProperties->Skeleton = AnimSequence->GetSkeleton();
 			CommonSkeletalMeshesAndAnimationsProperties->bImportOnlyAnimations = true;
 		}
+		else
+		{
+			HideCategories.Add(TEXT("Animations"));
+		}
+	}
+
+	if (UInterchangePipelineBase* OuterMostPipeline = GetMostPipelineOuter())
+	{
+		for (const FString& HideCategoryName : HideCategories)
+		{
+			HidePropertiesOfCategory(OuterMostPipeline, this, HideCategoryName);
+		}
 	}
 }
 
diff --git a/Engine/Plugins/Experimental/Interchange/Runtime/Source/Pipelines/Private/InterchangeGenericAssetsPipeline.cpp b/Engine/Plugins/Experimental/Interchange/Runtime/Source/Pipelines/Private/InterchangeGenericAssetsPipeline.cpp
--- a/Engine/Plugins/Experimental/Interchange/Runtime/Source/Pipelines/Private/InterchangeGenericAssetsPipeline.cpp
+++ b/Engine/Plugins/Experimental/Interchange/Runtime/Source/Pipelines/Private/InterchangeGenericAssetsPipeline.cpp
@@ -107,26 +107,28 @@ bool UInterchangeGenericAssetsPipeline::IsSettingsAreValid(TOptional<FText>& Out
 	return Super::IsSettingsAreValid(OutInvalidReason);
 }
 
-void UInterchangeGenericAssetsPipeline::AdjustSettingsForReimportType(EInterchangeReimportType ImportType, TObjectPtr<UObject> ReimportAsset)
+void UInterchangeGenericAssetsPipeline::AdjustSettingsForContext(EInterchangePipelineContext ImportType, TObjectPtr<UObject> ReimportAsset)
 {
+	Super::AdjustSettingsForContext(ImportType, ReimportAsset);
+
 	if (TexturePipeline)
 	{
-		TexturePipeline->AdjustSettingsForReimportType(ImportType, ReimportAsset);
+		TexturePipeline->AdjustSettingsForContext(ImportType, ReimportAsset);
 	}
 
 	if (MaterialPipeline)
 	{
-		MaterialPipeline->AdjustSettingsForReimportType(ImportType, ReimportAsset);
+		MaterialPipeline->AdjustSettingsForContext(ImportType, ReimportAsset);
 	}
 
 	if (MeshPipeline)
 	{
-		MeshPipeline->AdjustSettingsForReimportType(ImportType, ReimportAsset);
+		MeshPipeline->AdjustSettingsForContext(ImportType, ReimportAsset);
 	}
 
 	if (AnimationPipeline)
 	{
-		AnimationPipeline->AdjustSettingsForReimportType(ImportType, ReimportAsset);
+		AnimationPipeline->AdjustSettingsForContext(ImportType, ReimportAsset);
 	}
 }
 
diff --git a/Engine/Plugins/Experimental/Interchange/Runtime/Source/Pipelines/Private/InterchangeGenericMaterialPipeline.cpp b/Engine/Plugins/Experimental/Interchange/Runtime/Source/Pipelines/Private/InterchangeGenericMaterialPipeline.cpp
--- a/Engine/Plugins/Experimental/Interchange/Runtime/Source/Pipelines/Private/InterchangeGenericMaterialPipeline.cpp
+++ b/Engine/Plugins/Experimental/Interchange/Runtime/Source/Pipelines/Private/InterchangeGenericMaterialPipeline.cpp
@@ -42,16 +42,28 @@
 #include "UObject/Object.h"
 #include "UObject/ObjectMacros.h"
 
-void UInterchangeGenericMaterialPipeline::AdjustSettingsForReimportType(EInterchangeReimportType ImportType, TObjectPtr<UObject> ReimportAsset)
+void UInterchangeGenericMaterialPipeline::AdjustSettingsForContext(EInterchangePipelineContext ImportType, TObjectPtr<UObject> ReimportAsset)
 {
+	Super::AdjustSettingsForContext(ImportType, ReimportAsset);
+
+	TArray<FString> HideCategories;
 	bool bIsObjectAMaterial = ReimportAsset.IsNull() ? false : ReimportAsset->IsA(UMaterialInterface::StaticClass());
-	if ((!bIsObjectAMaterial && ImportType == EInterchangeReimportType::AssetReimport)
-		|| ImportType == EInterchangeReimportType::AssetCustomLODImport
-		|| ImportType == EInterchangeReimportType::AssetCustomLODReimport
-		|| ImportType == EInterchangeReimportType::AssetAlternateSkinningImport
-		|| ImportType == EInterchangeReimportType::AssetAlternateSkinningReimport)
+	if ((!bIsObjectAMaterial && ImportType == EInterchangePipelineContext::AssetReimport)
+		|| ImportType == EInterchangePipelineContext::AssetCustomLODImport
+		|| ImportType == EInterchangePipelineContext::AssetCustomLODReimport
+		|| ImportType == EInterchangePipelineContext::AssetAlternateSkinningImport
+		|| ImportType == EInterchangePipelineContext::AssetAlternateSkinningReimport)
 	{
 		MaterialImport = EInterchangeMaterialImportOption::DoNotImport;
+		HideCategories.Add(TEXT("Materials"));
+	}
+
+	if (UInterchangePipelineBase* OuterMostPipeline = GetMostPipelineOuter())
+	{
+		for (const FString& HideCategoryName : HideCategories)
+		{
+			HidePropertiesOfCategory(OuterMostPipeline, this, HideCategoryName);
+		}
 	}
 }
 
diff --git a/Engine/Plugins/Experimental/Interchange/Runtime/Source/Pipelines/Private/InterchangeGenericMeshPipeline.cpp b/Engine/Plugins/Experimental/Interchange/Runtime/Source/Pipelines/Private/InterchangeGenericMeshPipeline.cpp
--- a/Engine/Plugins/Experimental/Interchange/Runtime/Source/Pipelines/Private/InterchangeGenericMeshPipeline.cpp
+++ b/Engine/Plugins/Experimental/Interchange/Runtime/Source/Pipelines/Private/InterchangeGenericMeshPipeline.cpp
@@ -1,7 +1,9 @@
 // Copyright Epic Games, Inc. All Rights Reserved. 
 #include "InterchangeGenericMeshPipeline.h"
 
-#include "CoreMinimal.h"
+#include "Animation/AnimSequence.h"
+#include "Engine/SkeletalMesh.h"
+#include "Engine/StaticMesh.h"
 #include "InterchangeGenericAssetsPipeline.h"
 #include "InterchangeMeshNode.h"
 #include "InterchangePipelineLog.h"
@@ -13,21 +15,67 @@
 #include "Nodes/InterchangeBaseNode.h"
 #include "Nodes/InterchangeBaseNodeContainer.h"
 
-void UInterchangeGenericMeshPipeline::AdjustSettingsForReimportType(EInterchangeReimportType ImportType, TObjectPtr<UObject> ReimportAsset)
+void UInterchangeGenericMeshPipeline::AdjustSettingsForContext(EInterchangePipelineContext ImportType, TObjectPtr<UObject> ReimportAsset)
 {
+	Super::AdjustSettingsForContext(ImportType, ReimportAsset);
+
 	check(!CommonSkeletalMeshesAndAnimationsProperties.IsNull());
-	if(ImportType == EInterchangeReimportType::AssetCustomLODImport
-		|| ImportType == EInterchangeReimportType::AssetCustomLODReimport
-		|| ImportType == EInterchangeReimportType::AssetAlternateSkinningImport
-		|| ImportType == EInterchangeReimportType::AssetAlternateSkinningReimport)
+	if (ImportType == EInterchangePipelineContext::None)
+	{
+		//We do not change the setting if we are in editing context
+		return;
+	}
+
+	const bool bIsReimport = IsReimportContext();
+
+	//Avoid creating physics asset when importing a LOD or the alternate skinning
+	if (ImportType == EInterchangePipelineContext::AssetCustomLODImport
+		|| ImportType == EInterchangePipelineContext::AssetCustomLODReimport
+		|| ImportType == EInterchangePipelineContext::AssetAlternateSkinningImport
+		|| ImportType == EInterchangePipelineContext::AssetAlternateSkinningReimport)
 	{
 		bCreatePhysicsAsset = false;
 		PhysicsAsset = nullptr;
 	}
-	else if (USkeletalMesh* SkeletalMesh = Cast<USkeletalMesh>(ReimportAsset))
+
+	TArray<FString> HideCategories;
+	if (ImportType == EInterchangePipelineContext::AssetReimport)
+	{
+		if (USkeletalMesh* SkeletalMesh = Cast<USkeletalMesh>(ReimportAsset))
+		{
+
+			//Set the skeleton to the current asset skeleton
+			CommonSkeletalMeshesAndAnimationsProperties->Skeleton = SkeletalMesh->GetSkeleton();
+			bImportStaticMeshes = false;
+			HideCategories.Add(TEXT("Static Meshes"));
+		}
+		else if (UStaticMesh* StaticMesh = Cast<UStaticMesh>(ReimportAsset))
+		{
+			HideCategories.Add(TEXT("Skeletal Meshes"));
+			HideCategories.Add(TEXT("Common Skeletal Meshes and Animations"));
+		}
+		else if (UAnimSequence* AnimSequence = Cast<UAnimSequence>(ReimportAsset))
+		{
+			HideCategories.Add(TEXT("Static Meshes"));
+			HideCategories.Add(TEXT("Skeletal Meshes"));
+			HideCategories.Add(TEXT("Common Meshes"));
+		}
+		else if (ReimportAsset)
+		{
+			HideCategories.Add(TEXT("Static Meshes"));
+			HideCategories.Add(TEXT("Skeletal Meshes"));
+			HideCategories.Add(TEXT("Common Meshes"));
+			HideCategories.Add(TEXT("Common Skeletal Meshes and Animations"));
+		}
+	}
+
+	if (UInterchangePipelineBase* OuterMostPipeline = GetMostPipelineOuter())
 	{
-		//Set the skeleton to the current asset skeleton
-		CommonSkeletalMeshesAndAnimationsProperties->Skeleton = SkeletalMesh->GetSkeleton();
+		for (const FString& HideCategoryName : HideCategories)
+		{
+			constexpr bool bDoTransientSubPipeline = true;
+			HidePropertiesOfCategory(OuterMostPipeline, this, HideCategoryName, bDoTransientSubPipeline);
+		}
 	}
 }
 
diff --git a/Engine/Plugins/Experimental/Interchange/Runtime/Source/Pipelines/Private/InterchangeGenericTexturePipeline.cpp b/Engine/Plugins/Experimental/Interchange/Runtime/Source/Pipelines/Private/InterchangeGenericTexturePipeline.cpp
--- a/Engine/Plugins/Experimental/Interchange/Runtime/Source/Pipelines/Private/InterchangeGenericTexturePipeline.cpp
+++ b/Engine/Plugins/Experimental/Interchange/Runtime/Source/Pipelines/Private/InterchangeGenericTexturePipeline.cpp
@@ -94,16 +94,27 @@ namespace UE::Interchange::Private
 	}
 }
 
-void UInterchangeGenericTexturePipeline::AdjustSettingsForReimportType(EInterchangeReimportType ImportType, TObjectPtr<UObject> ReimportAsset)
+void UInterchangeGenericTexturePipeline::AdjustSettingsForContext(EInterchangePipelineContext ImportType, TObjectPtr<UObject> ReimportAsset)
 {
+	Super::AdjustSettingsForContext(ImportType, ReimportAsset);
+
+	TArray<FString> HideCategories;
 	bool bIsObjectATexture = ReimportAsset.IsNull() ? false : ReimportAsset.IsA(UTexture::StaticClass());
-	if( (!bIsObjectATexture && ImportType == EInterchangeReimportType::AssetReimport)
-		|| ImportType == EInterchangeReimportType::AssetCustomLODImport
-		|| ImportType == EInterchangeReimportType::AssetCustomLODReimport
-		|| ImportType == EInterchangeReimportType::AssetAlternateSkinningImport
-		|| ImportType == EInterchangeReimportType::AssetAlternateSkinningReimport)
+	if( (!bIsObjectATexture && ImportType == EInterchangePipelineContext::AssetReimport)
+		|| ImportType == EInterchangePipelineContext::AssetCustomLODImport
+		|| ImportType == EInterchangePipelineContext::AssetCustomLODReimport
+		|| ImportType == EInterchangePipelineContext::AssetAlternateSkinningImport
+		|| ImportType == EInterchangePipelineContext::AssetAlternateSkinningReimport)
 	{
 		bImportTextures = false;
+		HideCategories.Add(TEXT("Textures"));
+	}
+	if (UInterchangePipelineBase* OuterMostPipeline = GetMostPipelineOuter())
+	{
+		for (const FString& HideCategoryName : HideCategories)
+		{
+			HidePropertiesOfCategory(OuterMostPipeline, this, HideCategoryName);
+		}
 	}
 }
 
diff --git a/Engine/Plugins/Experimental/Interchange/Runtime/Source/Pipelines/Public/InterchangeGenericAnimationPipeline.h b/Engine/Plugins/Experimental/Interchange/Runtime/Source/Pipelines/Public/InterchangeGenericAnimationPipeline.h
--- a/Engine/Plugins/Experimental/Interchange/Runtime/Source/Pipelines/Public/InterchangeGenericAnimationPipeline.h
+++ b/Engine/Plugins/Experimental/Interchange/Runtime/Source/Pipelines/Public/InterchangeGenericAnimationPipeline.h
@@ -101,7 +101,7 @@ public:
 	UPROPERTY()
 	FString SourceAnimationName;
 
-	virtual void AdjustSettingsForReimportType(EInterchangeReimportType ImportType, TObjectPtr<UObject> ReimportAsset) override;
+	virtual void AdjustSettingsForContext(EInterchangePipelineContext ImportType, TObjectPtr<UObject> ReimportAsset) override;
 	
 protected:
 
diff --git a/Engine/Plugins/Experimental/Interchange/Runtime/Source/Pipelines/Public/InterchangeGenericAssetsPipeline.h b/Engine/Plugins/Experimental/Interchange/Runtime/Source/Pipelines/Public/InterchangeGenericAssetsPipeline.h
--- a/Engine/Plugins/Experimental/Interchange/Runtime/Source/Pipelines/Public/InterchangeGenericAssetsPipeline.h
+++ b/Engine/Plugins/Experimental/Interchange/Runtime/Source/Pipelines/Public/InterchangeGenericAssetsPipeline.h
@@ -81,7 +81,7 @@ public:
 	virtual bool IsSettingsAreValid(TOptional<FText>& OutInvalidReason) const override;
 
 
-	virtual void AdjustSettingsForReimportType(EInterchangeReimportType ImportType, TObjectPtr<UObject> ReimportAsset) override;
+	virtual void AdjustSettingsForContext(EInterchangePipelineContext ImportType, TObjectPtr<UObject> ReimportAsset) override;
 protected:
 
 	virtual void ExecutePreImportPipeline(UInterchangeBaseNodeContainer* InBaseNodeContainer, const TArray<UInterchangeSourceData*>& InSourceDatas) override;
diff --git a/Engine/Plugins/Experimental/Interchange/Runtime/Source/Pipelines/Public/InterchangeGenericMaterialPipeline.h b/Engine/Plugins/Experimental/Interchange/Runtime/Source/Pipelines/Public/InterchangeGenericMaterialPipeline.h
--- a/Engine/Plugins/Experimental/Interchange/Runtime/Source/Pipelines/Public/InterchangeGenericMaterialPipeline.h
+++ b/Engine/Plugins/Experimental/Interchange/Runtime/Source/Pipelines/Public/InterchangeGenericMaterialPipeline.h
@@ -42,7 +42,7 @@ public:
 	UPROPERTY(EditAnywhere, BlueprintReadWrite, Category = "Materials", Meta= (EditCondition="MaterialImport==EInterchangeMaterialImportOption::ImportAsMaterialInstances", AllowedClasses="/Script/Engine.MaterialInterface"))
 	FSoftObjectPath ParentMaterial;
 
-	virtual void AdjustSettingsForReimportType(EInterchangeReimportType ImportType, TObjectPtr<UObject> ReimportAsset) override;
+	virtual void AdjustSettingsForContext(EInterchangePipelineContext ImportType, TObjectPtr<UObject> ReimportAsset) override;
 
 protected:
 	virtual void ExecutePreImportPipeline(UInterchangeBaseNodeContainer* InBaseNodeContainer, const TArray<UInterchangeSourceData*>& InSourceDatas) override;
diff --git a/Engine/Plugins/Experimental/Interchange/Runtime/Source/Pipelines/Public/InterchangeGenericMeshPipeline.h b/Engine/Plugins/Experimental/Interchange/Runtime/Source/Pipelines/Public/InterchangeGenericMeshPipeline.h
--- a/Engine/Plugins/Experimental/Interchange/Runtime/Source/Pipelines/Public/InterchangeGenericMeshPipeline.h
+++ b/Engine/Plugins/Experimental/Interchange/Runtime/Source/Pipelines/Public/InterchangeGenericMeshPipeline.h
@@ -2,7 +2,6 @@
 
 #pragma once
 
-#include "CoreMinimal.h"
 #include "InterchangeGenericAssetsPipelineSharedSettings.h"
 #include "InterchangePipelineBase.h"
 #include "InterchangeSourceData.h"
@@ -106,7 +105,7 @@ public:
 	UPROPERTY(EditAnywhere, BlueprintReadWrite, Category = "Skeletal Meshes", meta = (editcondition = "!bCreatePhysicsAsset"))
 	TObjectPtr<UPhysicsAsset> PhysicsAsset;
 
-	virtual void AdjustSettingsForReimportType(EInterchangeReimportType ImportType, TObjectPtr<UObject> ReimportAsset) override;
+	virtual void AdjustSettingsForContext(EInterchangePipelineContext ImportType, TObjectPtr<UObject> ReimportAsset) override;
 
 	virtual void PreDialogCleanup(const FName PipelineStackName) override;
 
diff --git a/Engine/Plugins/Experimental/Interchange/Runtime/Source/Pipelines/Public/InterchangeGenericTexturePipeline.h b/Engine/Plugins/Experimental/Interchange/Runtime/Source/Pipelines/Public/InterchangeGenericTexturePipeline.h
--- a/Engine/Plugins/Experimental/Interchange/Runtime/Source/Pipelines/Public/InterchangeGenericTexturePipeline.h
+++ b/Engine/Plugins/Experimental/Interchange/Runtime/Source/Pipelines/Public/InterchangeGenericTexturePipeline.h
@@ -60,7 +60,7 @@ public:
 	bool bAllowNonPowerOfTwo = false;
 
 public:
-	virtual void AdjustSettingsForReimportType(EInterchangeReimportType ImportType, TObjectPtr<UObject> ReimportAsset) override;
+	virtual void AdjustSettingsForContext(EInterchangePipelineContext ImportType, TObjectPtr<UObject> ReimportAsset) override;
 	virtual void ExecutePreImportPipeline(UInterchangeBaseNodeContainer* InBaseNodeContainer, const TArray<UInterchangeSourceData*>& InSourceDatas) override;
 	virtual void ExecutePostImportPipeline(const UInterchangeBaseNodeContainer* InBaseNodeContainer, const FString& NodeKey, UObject* CreatedAsset, bool bIsAReimport) override;
 
diff --git a/Engine/Plugins/Tests/EditorTests/Source/EditorTests/Private/UnrealEd/FbxAutomationTests.cpp b/Engine/Plugins/Tests/EditorTests/Source/EditorTests/Private/UnrealEd/FbxAutomationTests.cpp
--- a/Engine/Plugins/Tests/EditorTests/Source/EditorTests/Private/UnrealEd/FbxAutomationTests.cpp
+++ b/Engine/Plugins/Tests/EditorTests/Source/EditorTests/Private/UnrealEd/FbxAutomationTests.cpp
@@ -580,7 +580,7 @@ bool FFbxImportAssetsAutomationTest::RunTest(const FString& Parameters)
 					FScopedSkeletalMeshPostEditChange ScopedPostEditChange(ExistingSkeletalMesh);
 					const FName ProfileName(TEXT("Alternate"));
 					const int32 LodIndexZero = 0;
-					const bool bResult = FSkinWeightsUtilities::ImportAlternateSkinWeight(ExistingSkeletalMesh, AltFile, LodIndexZero, ProfileName);
+					const bool bResult = FSkinWeightsUtilities::ImportAlternateSkinWeight(ExistingSkeletalMesh, AltFile, LodIndexZero, ProfileName, false);
 					if (!bResult)
 					{
 						ExecutionInfo.AddError(FString::Printf(TEXT("%s: Error adding alternate skinning file %s."), *CleanFilename, *AltFile));
diff --git a/Engine/Source/Editor/Persona/Private/SkinWeightProfileHelpers.cpp b/Engine/Source/Editor/Persona/Private/SkinWeightProfileHelpers.cpp
--- a/Engine/Source/Editor/Persona/Private/SkinWeightProfileHelpers.cpp
+++ b/Engine/Source/Editor/Persona/Private/SkinWeightProfileHelpers.cpp
@@ -74,7 +74,7 @@ void FSkinWeightProfileHelpers::ImportSkinWeightProfile(USkeletalMesh* InSkeleta
 			FScopedSkeletalMeshPostEditChange ScopedPostEditChange(InSkeletalMesh);
 			const FName ProfileName(*ImportSettings->ProfileName);			
 			// Try and import the skin weight profile from the provided FBX file path
-			const bool bResult = FSkinWeightsUtilities::ImportAlternateSkinWeight(InSkeletalMesh, PickedFileName, ImportSettings->LODIndex, ProfileName);
+			const bool bResult = FSkinWeightsUtilities::ImportAlternateSkinWeight(InSkeletalMesh, PickedFileName, ImportSettings->LODIndex, ProfileName, false);
 			if (bResult)
 			{
 				FNotificationInfo NotificationInfo(FText::GetEmpty());
@@ -107,7 +107,7 @@ void FSkinWeightProfileHelpers::ImportSkinWeightProfileLOD(USkeletalMesh* InSkel
 		FScopedSuspendAlternateSkinWeightPreview ScopedSuspendAlternateSkinnWeightPreview(InSkeletalMesh);
 		FScopedSkeletalMeshPostEditChange ScopedPostEditChange(InSkeletalMesh);
 		// Try and import skin weights for a specific mesh LOD
-		const bool bResult = FSkinWeightsUtilities::ImportAlternateSkinWeight(InSkeletalMesh, PickedFileName, LODIndex, ProfileName);
+		const bool bResult = FSkinWeightsUtilities::ImportAlternateSkinWeight(InSkeletalMesh, PickedFileName, LODIndex, ProfileName, false);
 		if (bResult)
 		{
 			if(!InSkeletalMesh->IsLODImportedDataBuildAvailable(LODIndex))
@@ -146,7 +146,7 @@ void FSkinWeightProfileHelpers::ReimportSkinWeightProfileLOD(USkeletalMesh* InSk
 			// Check to see if the source file is still valid
 			if (FPaths::FileExists(PathName))
 			{
-				bResult = FSkinWeightsUtilities::ImportAlternateSkinWeight(InSkeletalMesh, PathName, LODIndex, InProfileName);
+				bResult = FSkinWeightsUtilities::ImportAlternateSkinWeight(InSkeletalMesh, PathName, LODIndex, InProfileName, true);
 			}
 			else
 			{
@@ -155,7 +155,7 @@ void FSkinWeightProfileHelpers::ReimportSkinWeightProfileLOD(USkeletalMesh* InSk
 				{
 					// Otherwise let the user pick a new path
 					const FString PickedFileName = FSkinWeightsUtilities::PickSkinWeightPath(LODIndex, InSkeletalMesh);
-					bResult = FSkinWeightsUtilities::ImportAlternateSkinWeight(InSkeletalMesh, PickedFileName, LODIndex, InProfileName);
+					bResult = FSkinWeightsUtilities::ImportAlternateSkinWeight(InSkeletalMesh, PickedFileName, LODIndex, InProfileName, true);
 				}
 			}
 
diff --git a/Engine/Source/Editor/UnrealEd/Private/SkinWeightsUtilities.cpp b/Engine/Source/Editor/UnrealEd/Private/SkinWeightsUtilities.cpp
--- a/Engine/Source/Editor/UnrealEd/Private/SkinWeightsUtilities.cpp
+++ b/Engine/Source/Editor/UnrealEd/Private/SkinWeightsUtilities.cpp
@@ -34,6 +34,7 @@
 #include "Framework/Application/SlateApplication.h"
 #include "InterchangeAssetImportData.h"
 #include "InterchangeFilePickerBase.h"
+#include "InterchangePipelineBase.h"
 #include "InterchangeManager.h"
 #include "InterchangeProjectSettings.h"
 #include "Interfaces/ITargetPlatform.h"
@@ -43,7 +44,7 @@
 
 DEFINE_LOG_CATEGORY_STATIC(LogSkinWeightsUtilities, Log, All);
 
-bool FSkinWeightsUtilities::ImportAlternateSkinWeight(USkeletalMesh* SkeletalMesh, const FString& Path, int32 TargetLODIndex, const FName& ProfileName)
+bool FSkinWeightsUtilities::ImportAlternateSkinWeight(USkeletalMesh* SkeletalMesh, const FString& Path, int32 TargetLODIndex, const FName& ProfileName, const bool bIsReimport)
 {
 	check(SkeletalMesh);
 	check(SkeletalMesh->GetLODInfo(TargetLODIndex));
@@ -106,10 +107,19 @@ bool FSkinWeightsUtilities::ImportAlternateSkinWeight(USkeletalMesh* SkeletalMes
 	if (bUseInterchangeFramework)
 	{
 		UE::Interchange::FScopedSourceData ScopedSourceData(AbsoluteFilePath);
-
+		const UInterchangeProjectSettings* InterchangeProjectSettings = GetDefault<UInterchangeProjectSettings>();
 		FImportAssetParameters ImportAssetParameters;
 		ImportAssetParameters.bIsAutomated = true; // From the InterchangeManager point of view, this is considered an automated import
 
+		if (const UClass* GenericPipelineClass = InterchangeProjectSettings->GenericPipelineClass.LoadSynchronous())
+		{
+			if (UInterchangePipelineBase* GenericPipeline = NewObject<UInterchangePipelineBase>(GetTransientPackage(), GenericPipelineClass))
+			{
+				GenericPipeline->AdjustSettingsForContext(bIsReimport ? EInterchangePipelineContext::AssetAlternateSkinningReimport : EInterchangePipelineContext::AssetAlternateSkinningImport, nullptr);
+				ImportAssetParameters.OverridePipelines.Add(GenericPipeline);
+			}
+		}
+
 		//TODO create a pipeline that set all the proper skeletalmesh options (look at the legacy system setup)
 		UE::Interchange::FAssetImportResultRef AssetImportResult = InterchangeManager.ImportAssetAsync(ImportAssetPath, ScopedSourceData.GetSourceData(), ImportAssetParameters);
 		AssetImportResult->WaitUntilDone(); //TODO, do not stall the main thread here, WaitUntilDone will tick taskgraph so the job can complete even if we wait on the game thread.
@@ -223,13 +233,13 @@ bool FSkinWeightsUtilities::ImportAlternateSkinWeight(USkeletalMesh* SkeletalMes
 				// Prepare the profile data
 				FSkinWeightProfileInfo* Profile = SkeletalMesh->GetSkinWeightProfiles().FindByPredicate([ProfileName](FSkinWeightProfileInfo Profile) { return Profile.Name == ProfileName; });
 
-				const bool bIsReimport = Profile != nullptr;
-				FText TransactionName = bIsReimport ? NSLOCTEXT("UnrealEd", "UpdateAlternateSkinningWeight", "Update Alternate Skinning Weight")
+				const bool bIsReimportLocal = Profile != nullptr;
+				FText TransactionName = bIsReimportLocal ? NSLOCTEXT("UnrealEd", "UpdateAlternateSkinningWeight", "Update Alternate Skinning Weight")
 					: NSLOCTEXT("UnrealEd", "ImportAlternateSkinningWeight", "Import Alternate Skinning Weight");
 				FScopedTransaction ScopedTransaction(TransactionName);
 				SkeletalMesh->Modify();
 
-				if (bIsReimport)
+				if (bIsReimportLocal)
 				{
 					// Update source file path
 					FString& StoredPath = Profile->PerLODSourceFiles.FindOrAdd(TargetLODIndex);
@@ -314,14 +324,14 @@ bool FSkinWeightsUtilities::ReimportAlternateSkinWeight(USkeletalMesh* SkeletalM
 		FString AbsoluteFilePath = UAssetImportData::ResolveImportFilename(PathName, SkeletalMesh->GetOutermost());
 		if (FPaths::FileExists(AbsoluteFilePath))
 		{
-			bResult |= FSkinWeightsUtilities::ImportAlternateSkinWeight(SkeletalMesh, AbsoluteFilePath, TargetLODIndex, ProfileInfo.Name);
+			bResult |= FSkinWeightsUtilities::ImportAlternateSkinWeight(SkeletalMesh, AbsoluteFilePath, TargetLODIndex, ProfileInfo.Name, true);
 		}
 		else
 		{
 			const FString PickedFileName = FSkinWeightsUtilities::PickSkinWeightPath(TargetLODIndex, SkeletalMesh);
 			if (!PickedFileName.IsEmpty() && FPaths::FileExists(PickedFileName))
 			{
-				bResult |= FSkinWeightsUtilities::ImportAlternateSkinWeight(SkeletalMesh, PickedFileName, TargetLODIndex, ProfileInfo.Name);
+				bResult |= FSkinWeightsUtilities::ImportAlternateSkinWeight(SkeletalMesh, PickedFileName, TargetLODIndex, ProfileInfo.Name, true);
 			}
 		}
 	}
diff --git a/Engine/Source/Editor/UnrealEd/Public/SkinWeightsUtilities.h b/Engine/Source/Editor/UnrealEd/Public/SkinWeightsUtilities.h
--- a/Engine/Source/Editor/UnrealEd/Public/SkinWeightsUtilities.h
+++ b/Engine/Source/Editor/UnrealEd/Public/SkinWeightsUtilities.h
@@ -27,7 +27,7 @@ public:
 	* @Param ProfileName - The name of the profile to associate the imported skin weight
 	* @param bReregisterComponent - if true the component using the skeletal mesh will all be re register.
 	*/
-	static bool ImportAlternateSkinWeight(USkeletalMesh* SkeletalMesh, const FString& Path, int32 TargetLODIndex, const FName& ProfileName);
+	static bool ImportAlternateSkinWeight(USkeletalMesh* SkeletalMesh, const FString& Path, int32 TargetLODIndex, const FName& ProfileName, const bool bIsReimport);
 	
 	/**
 	* This function reimport all skin weights profile for a specified LOD. Return true if the weights are successfully updates.
diff --git a/Engine/Source/Runtime/Experimental/Interchange/Core/Private/InterchangePipelineBase.cpp b/Engine/Source/Runtime/Experimental/Interchange/Core/Private/InterchangePipelineBase.cpp
--- a/Engine/Source/Runtime/Experimental/Interchange/Core/Private/InterchangePipelineBase.cpp
+++ b/Engine/Source/Runtime/Experimental/Interchange/Core/Private/InterchangePipelineBase.cpp
@@ -25,7 +25,7 @@ namespace UE
 
 void UInterchangePipelineBase::LoadSettings(const FName PipelineStackName)
 {
-	LoadSettingsInternal(PipelineStackName, GEditorPerProjectIni, LockedProperties);
+	LoadSettingsInternal(PipelineStackName, GEditorPerProjectIni, PropertiesStates);
 }
 
 void UInterchangePipelineBase::SaveSettings(const FName PipelineStackName)
@@ -33,38 +33,49 @@ void UInterchangePipelineBase::SaveSettings(const FName PipelineStackName)
 	SaveSettingsInternal(PipelineStackName, GEditorPerProjectIni);
 }
 
-bool UInterchangePipelineBase::SetLockedPropertyStatus(const FName PropertyPath, bool bLocked)
+void UInterchangePipelineBase::AdjustSettingsForContext(EInterchangePipelineContext ReimportType, TObjectPtr<UObject> ReimportAsset)
 {
-	ensure(bAllowLockedPropertiesEdition);
-
-	if (!bLocked)
+	bAllowPropertyStatesEdition = (ReimportType == EInterchangePipelineContext::None);
+	bIsReimportContext = false;
+	switch (ReimportType)
 	{
-		LockedProperties.Remove(PropertyPath);
-	}
-	else
+	case EInterchangePipelineContext::AssetReimport:
+	case EInterchangePipelineContext::AssetAlternateSkinningReimport:
+	case EInterchangePipelineContext::AssetCustomLODReimport:
+	case EInterchangePipelineContext::SceneReimport:
 	{
-		bool& bLockStatus = LockedProperties.FindOrAdd(PropertyPath);
-		bLockStatus = bLocked;
+		bIsReimportContext = true;
+	}
+	break;
 	}
+}
 
-	return true;
+const FInterchangePipelinePropertyStates* UInterchangePipelineBase::GetPropertyStates(const FName PropertyPath) const
+{
+	return PropertiesStates.Find(PropertyPath);
 }
 
-bool UInterchangePipelineBase::GetLockedPropertyStatus(const FName PropertyPath) const
+FInterchangePipelinePropertyStates* UInterchangePipelineBase::GetMutablePropertyStates(const FName PropertyPath)
 {
-	if (const bool* bLockStatus = LockedProperties.Find(PropertyPath))
-	{
-		return *bLockStatus;
-	}
-	return false;
+	return PropertiesStates.Find(PropertyPath);
 }
 
-FName UInterchangePipelineBase::GetLockedPropertiesPropertyName()
+bool UInterchangePipelineBase::DoesPropertyStatesExist(const FName PropertyPath) const
 {
-	return GET_MEMBER_NAME_CHECKED(UInterchangePipelineBase, LockedProperties);
+	return PropertiesStates.Contains(PropertyPath);
 }
 
-void UInterchangePipelineBase::LoadSettingsInternal(const FName PipelineStackName, const FString& ConfigFilename, TMap<FName, bool>& ParentLockedProperties)
+FInterchangePipelinePropertyStates& UInterchangePipelineBase::FindOrAddPropertyStates(const FName PropertyPath)
+{
+	return PropertiesStates.FindOrAdd(PropertyPath);
+}
+
+FName UInterchangePipelineBase::GetPropertiesStatesPropertyName()
+{
+	return GET_MEMBER_NAME_CHECKED(UInterchangePipelineBase, PropertiesStates);
+}
+
+void UInterchangePipelineBase::LoadSettingsInternal(const FName PipelineStackName, const FString& ConfigFilename, TMap<FName, FInterchangePipelinePropertyStates>& ParentPropertiesStates)
 {
 	int32 PortFlags = 0;
 	UClass* Class = this->GetClass();
@@ -80,14 +91,14 @@ void UInterchangePipelineBase::LoadSettingsInternal(const FName PipelineStackNam
 		FString Key = Property->GetName();
 		const FName PropertyName = Property->GetFName();
 		const FName PropertyPath = FName(Property->GetPathName());
-		if (PropertyName == GET_MEMBER_NAME_CHECKED(UInterchangePipelineBase, LockedProperties))
+		if (PropertyName == GET_MEMBER_NAME_CHECKED(UInterchangePipelineBase, PropertiesStates))
 		{
 			continue;
 		}
 
-		if (const bool* LockStatus = ParentLockedProperties.Find(PropertyPath))
+		if (const FInterchangePipelinePropertyStates* PropertyStates = ParentPropertiesStates.Find(PropertyPath))
 		{
-			if (*LockStatus)
+			if (PropertyStates->IsPropertyLocked())
 			{
 				//Skip this locked property
 				continue;
@@ -155,7 +166,7 @@ void UInterchangePipelineBase::LoadSettingsInternal(const FName PipelineStackNam
 			// Load the settings if the referenced pipeline is a subobject of ours
 			if (SubPipeline->IsInOuter(this))
 			{
-				SubPipeline->LoadSettingsInternal(PipelineStackName, ConfigFilename, ParentLockedProperties);
+				SubPipeline->LoadSettingsInternal(PipelineStackName, ConfigFilename, ParentPropertiesStates);
 			}
 		}
 		else
@@ -198,7 +209,7 @@ void UInterchangePipelineBase::SaveSettingsInternal(const FName PipelineStackNam
 		FString SectionName = UE::Interchange::PipelinePrivate::CreateConfigSectionName(PipelineStackName, Class);
 		FString Key = Property->GetName();
 		const FName PropertyName = Property->GetFName();
-		if (PropertyName == GET_MEMBER_NAME_CHECKED(UInterchangePipelineBase, LockedProperties))
+		if (PropertyName == GET_MEMBER_NAME_CHECKED(UInterchangePipelineBase, PropertiesStates))
 		{
 			continue;
 		}
@@ -250,3 +261,53 @@ void UInterchangePipelineBase::SaveSettingsInternal(const FName PipelineStackNam
 	}
 	GConfig->Flush(0);
 }
+
+UInterchangePipelineBase* UInterchangePipelineBase::GetMostPipelineOuter() const
+{
+	UInterchangePipelineBase* Top = (UInterchangePipelineBase*)this;
+	for (;;)
+	{
+		if (UInterchangePipelineBase* CurrentOuter = Cast<UInterchangePipelineBase>(Top->GetOuter()))
+		{
+			Top = CurrentOuter;
+			continue;
+		}
+		break;
+	}
+	return Top;
+}
+
+void UInterchangePipelineBase::HidePropertiesOfCategory(UInterchangePipelineBase* OuterMostPipeline, UInterchangePipelineBase* Pipeline, const FString& HideCategoryName, bool bDoTransientSubPipeline /*= false*/)
+{
+#if WITH_EDITOR
+	UClass* PipelineClass = Pipeline->GetClass();
+	const FName CategoryKey("Category");
+	for (FProperty* Property = PipelineClass->PropertyLink; Property; Property = Property->PropertyLinkNext)
+	{
+		FObjectProperty* SubObject = CastField<FObjectProperty>(Property);
+		UInterchangePipelineBase* SubPipeline = SubObject ? Cast<UInterchangePipelineBase>(SubObject->GetObjectPropertyValue_InContainer(Pipeline)) : nullptr;
+		const bool bSkipTransient = !bDoTransientSubPipeline || !SubPipeline;
+		//Do not save a transient property
+		if (bSkipTransient && Property->HasAnyPropertyFlags(CPF_Transient))
+		{
+			continue;
+		}
+
+		const FName PropertyName = Property->GetFName();
+		const FName PropertyPath = FName(Property->GetPathName());
+		
+		if (SubPipeline)
+		{
+			HidePropertiesOfCategory(OuterMostPipeline, SubPipeline, HideCategoryName);
+		}
+		else
+		{
+			FString CategoryName = Property->GetMetaData(CategoryKey);
+			if (CategoryName.Equals(HideCategoryName))
+			{
+				OuterMostPipeline->FindOrAddPropertyStates(PropertyPath).ReimportStates.bVisible = false;
+			}
+		}
+	}
+#endif
+}
diff --git a/Engine/Source/Runtime/Experimental/Interchange/Core/Public/InterchangePipelineBase.h b/Engine/Source/Runtime/Experimental/Interchange/Core/Public/InterchangePipelineBase.h
--- a/Engine/Source/Runtime/Experimental/Interchange/Core/Public/InterchangePipelineBase.h
+++ b/Engine/Source/Runtime/Experimental/Interchange/Core/Public/InterchangePipelineBase.h
@@ -37,9 +37,12 @@ enum class EInterchangePipelineTask : uint8
 };
 
 UENUM()
-enum class EInterchangeReimportType : uint8
+enum class EInterchangePipelineContext : uint8
 {
+	None, //Default pipeline instance we refer in the project settings pipeline stack. This context should allow editing of the properties states
+	AssetImport,
 	AssetReimport,
+	SceneImport,
 	SceneReimport,
 	AssetCustomLODImport, //The import for custom LOD is there because we use a copy of the asset import data pipeline stack.
 	AssetCustomLODReimport,
@@ -47,6 +50,56 @@ enum class EInterchangeReimportType : uint8
 	AssetAlternateSkinningReimport,
 };
 
+USTRUCT(BlueprintType)
+struct FInterchangePipelinePropertyStatePerContext
+{
+	GENERATED_BODY()
+
+	/** If true, the property is visible. */
+	UPROPERTY(EditAnywhere, Category = "Property States")
+	bool bVisible = true;
+};
+
+USTRUCT(BlueprintType)
+struct FInterchangePipelinePropertyStates
+{
+	GENERATED_BODY()
+
+	/** Return true if the property is locked. */
+	bool IsPropertyLocked() const
+	{
+		return bLocked;
+	}
+	
+	void SetPropertyLocked(const bool bLockValue)
+	{
+		bLocked = bLockValue;
+	}
+	
+	/** Return true if the property is visible for the specified context. */
+	bool IsPropertyVisible(const bool bIsReimportContext) const
+	{
+		return bIsReimportContext ? ReimportStates.bVisible : ImportStates.bVisible;
+	}
+
+	void SetPropertyVisible(const bool bIsReimportContext, const bool bVisibleValue)
+	{
+		bIsReimportContext ? ReimportStates.bVisible = bVisibleValue : ImportStates.bVisible = bVisibleValue;
+	}
+
+	/** If true, the property is locked. */
+	UPROPERTY(EditAnywhere, Category = "Property States")
+	bool bLocked = false;
+
+	/** The property states for the import context */
+	UPROPERTY(EditAnywhere, Category = "Context Properties States")
+	FInterchangePipelinePropertyStatePerContext ImportStates;
+
+	/** The property states for the reimport context */
+	UPROPERTY(EditAnywhere, Category = "Context Properties States")
+	FInterchangePipelinePropertyStatePerContext ReimportStates;
+};
+
 UCLASS(BlueprintType, Blueprintable, editinlinenew, Experimental, Abstract)
 class INTERCHANGECORE_API UInterchangePipelineBase : public UObject
 {
@@ -141,7 +194,7 @@ public:
 	void SaveSettings(const FName PipelineStackName);
 
 	/**
-	 * This function is called before showing the import dialog it is not called when doing a re-import.
+	 * This function is called before showing the import dialog it is not called doing a re-import.
 	 */
 	virtual void PreDialogCleanup(const FName PipelineStackName) {}
 
@@ -155,12 +208,14 @@ public:
 	}
 
 	/**
-	 * This function is call only when we do a re-import before we show the pipeline dialog. Pipeline that override it can change the existing settings according to the re-import type.
+	 * This function is call when before we show the pipeline dialog. Pipeline that override it can change the existing settings according to the re-import type.
 	 * The function is also call when we import or re-import custom LOD and alternate skinning.
+	 * 
+	 * @Note - The function will set the context of the pipeline
 	 * @Param ReimportType - Tell pipeline what re-import type the user want to achieve.
 	 * @Param ReimportAsset - This is an optional parameter which is set when re-importing an asset.
 	 */
-	virtual void AdjustSettingsForReimportType(EInterchangeReimportType ReimportType, TObjectPtr<UObject> ReimportAsset) {}
+	virtual void AdjustSettingsForContext(EInterchangePipelineContext ReimportType, TObjectPtr<UObject> ReimportAsset);
 
 	/**
 	 * This function is used to add the given message object directly into the results for this operation.
@@ -184,23 +239,35 @@ public:
 		Results = InResults;
 	}
 
-	bool SetLockedPropertyStatus(const FName PropertyPath, bool bLocked);
-	bool GetLockedPropertyStatus(const FName PropertyPath) const;
+	/**
+	 * Return a const property states pointer. Return nullptr if the property do not exist
+	 */
+	const FInterchangePipelinePropertyStates* GetPropertyStates(const FName PropertyPath) const;
 
 	/**
-	 * Return the property name of the LockedProperties map
+	 * Return a mutable property states pointer. Return nullptr if the property do not exist
 	 */
-	static FName GetLockedPropertiesPropertyName();
+	FInterchangePipelinePropertyStates* GetMutablePropertyStates(const FName PropertyPath);
 
 	/**
-	 * If true, the property editor for this pipeline instance will allow locked properties edition.
-	 * If false, the property editor for this pipeline instance will set locked properties in read only.
-	 * 
-	 * Note: If you open in the content browser a pipeline asset you will be able to edit locked properties.
-	 *       If you import a file with interchange, the import dialog will show locked properties in read only mode.
+	 * Return true if the property has valid states, false if no states was set for the property
+	 */
+	UFUNCTION(BlueprintCallable, Category = "Interchange | Pipeline")
+	bool DoesPropertyStatesExist(const FName PropertyPath) const;
+
+	/**
+	 * Return a mutable property states Reference. Add the property states if it doesnt exist.
+	 */
+	UFUNCTION(BlueprintCallable, Category = "Interchange | Pipeline")
+	FInterchangePipelinePropertyStates& FindOrAddPropertyStates(const FName PropertyPath);
+
+	/**
+	 * Return the property name of the properties states map
 	 */
-	bool bAllowLockedPropertiesEdition = true;
+	static FName GetPropertiesStatesPropertyName();
 
+	bool CanEditPropertiesStates() { return bAllowPropertyStatesEdition; }
+	bool IsReimportContext() { return bIsReimportContext; }
 protected:
 
 	/**
@@ -241,14 +308,38 @@ protected:
 	{
 	}
 
-	void LoadSettingsInternal(const FName PipelineStackName, const FString& ConfigFilename, TMap<FName, bool>& ParentLockedProperties);
+	void LoadSettingsInternal(const FName PipelineStackName, const FString& ConfigFilename, TMap<FName, FInterchangePipelinePropertyStates>& ParentPropertiesStates);
 
 	void SaveSettingsInternal(const FName PipelineStackName, const FString& ConfigFilename);
 
+	UInterchangePipelineBase* GetMostPipelineOuter() const;
+
+	static void HidePropertiesOfCategory(UInterchangePipelineBase* OuterMostPipeline, UInterchangePipelineBase* Pipeline, const FString& HideCategoryName, bool bDoTransientSubPipeline = false);
+
+	/**
+	 * If true, the property editor for this pipeline instance will allow properties states edition.
+	 * If false, the property editor for this pipeline instance will apply the properties states.
+	 *
+	 * Note: If you open in the content browser a pipeline asset you will be able to edit the properties states.
+	 *       If you import a file with interchange, the import dialog will apply properties states.
+	 */
+	bool bAllowPropertyStatesEdition = true;
+
+	/**
+	 * If true, this pipeline instance is use for reimport.
+	 * If false, this pipeline instance is use for import.
+	 *
+	 * Note: This context must be set by the owner instancing this pipeline. This context will be use to hide or not some properties.
+	 */
+	bool bIsReimportContext = false;
+
 	UPROPERTY()
 	TObjectPtr<UInterchangeResultsContainer> Results;
 
-	/* Map of property path and lock status. Any properties that have a true lock status will be readonly when showing the import dialog */
-	UPROPERTY(BlueprintReadWrite, Category = "Pipeline Properties Manager")
-	TMap<FName, bool> LockedProperties;
+	/**
+	 * Map of property path and lock status. Any properties that have a true lock status will be readonly when showing the import dialog.
+	 * Use the API to Get and Set the properties states
+	*/
+	UPROPERTY()
+	TMap<FName, FInterchangePipelinePropertyStates> PropertiesStates;
 };
diff --git a/Engine/Source/Runtime/Experimental/Interchange/Engine/Private/InterchangeManager.cpp b/Engine/Source/Runtime/Experimental/Interchange/Engine/Private/InterchangeManager.cpp
--- a/Engine/Source/Runtime/Experimental/Interchange/Engine/Private/InterchangeManager.cpp
+++ b/Engine/Source/Runtime/Experimental/Interchange/Engine/Private/InterchangeManager.cpp
@@ -849,13 +849,13 @@ UInterchangeManager::ImportInternal(const FString& ContentPath, const UInterchan
 					//Duplicate the pipeline saved in the asset import data
 					UInterchangePipelineBase* GeneratedPipeline = Cast<UInterchangePipelineBase>(StaticDuplicateObject(SourcePipeline, GetTransientPackage()));
 
-					if (ImportType == UE::Interchange::EImportType::ImportType_Scene)
+					if (bImportScene)
 					{
-						GeneratedPipeline->AdjustSettingsForReimportType(EInterchangeReimportType::SceneReimport, nullptr);
+						GeneratedPipeline->AdjustSettingsForContext(EInterchangePipelineContext::SceneReimport, nullptr);
 					}
 					else
 					{
-						GeneratedPipeline->AdjustSettingsForReimportType(EInterchangeReimportType::AssetReimport, ImportAssetParameters.ReimportAsset);
+						GeneratedPipeline->AdjustSettingsForContext(EInterchangePipelineContext::AssetReimport, ImportAssetParameters.ReimportAsset);
 					}
 					PipelineStack.Add(GeneratedPipeline);
 				}
@@ -945,7 +945,7 @@ UInterchangeManager::ImportInternal(const FString& ContentPath, const UInterchan
 						{
 							if(UInterchangePipelineBase* GeneratedPipeline = UE::Interchange::GeneratePipelineInstance(Pipelines[GraphPipelineIndex]))
 							{
-								GeneratedPipeline->bAllowLockedPropertiesEdition = false;
+								GeneratedPipeline->AdjustSettingsForContext(bImportScene ? EInterchangePipelineContext::SceneImport : EInterchangePipelineContext::AssetImport, nullptr);
 								//Load the settings for this pipeline
 								//The dialog is saving the settings for all default pipelines
 								GeneratedPipeline->LoadSettings(PipelineStackName);
@@ -982,7 +982,14 @@ UInterchangeManager::ImportInternal(const FString& ContentPath, const UInterchan
 		{
 			// Duplicate the override pipelines to protect the scripted users form making race conditions
 			UInterchangePipelineBase* GeneratedPipeline = DuplicateObject<UInterchangePipelineBase>(ImportAssetParameters.OverridePipelines[GraphPipelineIndex], GetTransientPackage());
-			GeneratedPipeline->bAllowLockedPropertiesEdition = false;
+			if (OriginalAssetImportData != nullptr)
+			{
+				GeneratedPipeline->AdjustSettingsForContext(bImportScene ? EInterchangePipelineContext::SceneImport : EInterchangePipelineContext::AssetImport, ImportAssetParameters.ReimportAsset);
+			}
+			else
+			{
+				GeneratedPipeline->AdjustSettingsForContext(bImportScene ? EInterchangePipelineContext::SceneImport : EInterchangePipelineContext::AssetImport, nullptr);
+			}
 			AsyncHelper->Pipelines.Add(GeneratedPipeline);
 			AsyncHelper->OriginalPipelines.Add(GeneratedPipeline);
 		}
diff --git a/Engine/Source/Runtime/Experimental/Interchange/Engine/Private/InterchangeMeshUtilities.cpp b/Engine/Source/Runtime/Experimental/Interchange/Engine/Private/InterchangeMeshUtilities.cpp
--- a/Engine/Source/Runtime/Experimental/Interchange/Engine/Private/InterchangeMeshUtilities.cpp
+++ b/Engine/Source/Runtime/Experimental/Interchange/Engine/Private/InterchangeMeshUtilities.cpp
@@ -62,7 +62,7 @@ TFuture<bool> UInterchangeMeshUtilities::ImportCustomLodAsync(UObject* MeshObjec
 			//We set bAllowMultipleFile to false, we should have only one result
 			if (ensure(Filenames.Num() == 1))
 			{
-				const UInterchangeSourceData* SourceData = UInterchangeManager::GetInterchangeManager().CreateSourceData(Filenames[0]);
+				const UInterchangeSourceData* SourceData = InterchangeManager.CreateSourceData(Filenames[0]);
 				return InternalImportCustomLodAsync(Promise, MeshObject, LodIndex, SourceData);
 			}
 		}
@@ -87,13 +87,13 @@ TFuture<bool> UInterchangeMeshUtilities::InternalImportCustomLodAsync(TSharedPtr
 	UInterchangeAssetImportData* InterchangeAssetImportData = nullptr;
 	USkeletalMesh* SkeletalMesh = Cast<USkeletalMesh>(MeshObject);
 	UStaticMesh* StaticMesh = Cast<UStaticMesh>(MeshObject);
-	EInterchangeReimportType ImportType = EInterchangeReimportType::AssetCustomLODImport;
+	EInterchangePipelineContext ImportType = EInterchangePipelineContext::AssetCustomLODImport;
 	if (SkeletalMesh)
 	{
 		InterchangeAssetImportData = Cast<UInterchangeAssetImportData>(SkeletalMesh->GetAssetImportData());
 		if (SkeletalMesh->GetLODNum() > LodIndex)
 		{
-			ImportType = EInterchangeReimportType::AssetCustomLODReimport;
+			ImportType = EInterchangePipelineContext::AssetCustomLODReimport;
 		}
 	}
 	else if (StaticMesh)
@@ -101,7 +101,7 @@ TFuture<bool> UInterchangeMeshUtilities::InternalImportCustomLodAsync(TSharedPtr
 		InterchangeAssetImportData = Cast<UInterchangeAssetImportData>(StaticMesh->GetAssetImportData());
 		if (StaticMesh->GetNumSourceModels() > LodIndex)
 		{
-			ImportType = EInterchangeReimportType::AssetCustomLODReimport;
+			ImportType = EInterchangePipelineContext::AssetCustomLODReimport;
 		}
 	}
 	else
@@ -126,7 +126,7 @@ TFuture<bool> UInterchangeMeshUtilities::InternalImportCustomLodAsync(TSharedPtr
 		}
 		if (ensure(GeneratedPipeline))
 		{
-			GeneratedPipeline->AdjustSettingsForReimportType(ImportType, nullptr);
+			GeneratedPipeline->AdjustSettingsForContext(ImportType, nullptr);
 			ImportAssetParameters.OverridePipelines.Add(GeneratedPipeline);
 		}
 	}
diff --git a/Engine/Source/Runtime/Experimental/Interchange/Engine/Public/InterchangeProjectSettings.h b/Engine/Source/Runtime/Experimental/Interchange/Engine/Public/InterchangeProjectSettings.h
--- a/Engine/Source/Runtime/Experimental/Interchange/Engine/Public/InterchangeProjectSettings.h
+++ b/Engine/Source/Runtime/Experimental/Interchange/Engine/Public/InterchangeProjectSettings.h
@@ -76,4 +76,11 @@ public:
 	 */
 	UPROPERTY(EditAnywhere, config, Category = "Generic|ImportSettings")
 	bool bStaticMeshUseSmoothEdgesIfSmoothingInformationIsMissing = true;
+
+	/**
+	 * This tells interchange which is the pipeline class to use when editor tools want to import or reimport tools with bake settings.
+	 * UnrealEd code depend on this class to be set and this property is only editable in the ini file directly.
+	 */
+	UPROPERTY(EditAnywhere, config, Category = "Editor Generic Pipeline Class")
+	TSoftClassPtr <UInterchangePipelineBase> GenericPipelineClass;
 };
\ No newline at end of file
