Fix CIS Localization warning
#rnx

diff --git a/Engine/Plugins/Experimental/Interchange/Editor/Source/Pipelines/Private/InterchangeEditorPipelineDetails.cpp b/Engine/Plugins/Experimental/Interchange/Editor/Source/Pipelines/Private/InterchangeEditorPipelineDetails.cpp
--- a/Engine/Plugins/Experimental/Interchange/Editor/Source/Pipelines/Private/InterchangeEditorPipelineDetails.cpp
+++ b/Engine/Plugins/Experimental/Interchange/Editor/Source/Pipelines/Private/InterchangeEditorPipelineDetails.cpp
@@ -248,7 +248,7 @@ void FInterchangePipelineBaseDetailsCustomization::CustomizeDetails(IDetailLayou
 							{
 								return;
 							}
-							FScopedTransaction ScopedTransaction(NSLOCTEXT("InterchangePipelineBaseDetails::CustomizeDetails", "TransactionLockedPropertiesSet", "Toggle property locked at import"), !GIsTransacting);
+							FScopedTransaction ScopedTransaction(NSLOCTEXT("InterchangePipelineBaseDetails::CustomizeDetails", "TransactionTogglePropertyLocked", "Toggle property locked at import"), !GIsTransacting);
 							InterchangePipelinePtr->Modify();
 							InterchangePipelinePtr->FindOrAddPropertyStates(PropertyPath).SetPropertyLocked(CheckType == ECheckBoxState::Checked);
 							InterchangePipelinePtr->PostEditChange();
