Interchange
- Added a feature flag to toggle on and off custom pipelines support.
- Right now, it only hides the pipeline selection from the import dialog.

diff --git a/Engine/Plugins/Experimental/Interchange/Editor/Source/Pipelines/Private/SInterchangePipelineConfigurationDialog.cpp b/Engine/Plugins/Experimental/Interchange/Editor/Source/Pipelines/Private/SInterchangePipelineConfigurationDialog.cpp
--- a/Engine/Plugins/Experimental/Interchange/Editor/Source/Pipelines/Private/SInterchangePipelineConfigurationDialog.cpp
+++ b/Engine/Plugins/Experimental/Interchange/Editor/Source/Pipelines/Private/SInterchangePipelineConfigurationDialog.cpp
@@ -28,6 +28,8 @@
 
 #define LOCTEXT_NAMESPACE "InterchangePipelineConfiguration"
 
+extern INTERCHANGECORE_API bool GInterchangeEnableCustomPipelines;
+
 namespace UE::Interchange::Private
 {
 	const FInterchangeImportSettings& GetDefaultImportSettings(const bool bIsSceneImport)
@@ -152,6 +154,15 @@ void SInterchangePipelineStacksTreeView::Construct(const FArguments& InArgs)
 	);
 }
 
+void SInterchangePipelineStacksTreeView::SelectDefaultItem()
+{
+	if (!RootNodeArray.IsEmpty() && !RootNodeArray[0]->Childrens.IsEmpty())
+	{
+		SetItemExpansion(RootNodeArray[0], true);
+		SetSelection(RootNodeArray[0]->Childrens[0]);
+	}
+}
+
 /** The item used for visualizing the class in the tree. */
 class SInterchangePipelineStacksTreeViewItem : public STableRow< TSharedPtr<FInterchangePipelineStacksTreeNodeItem> >
 {
@@ -411,13 +422,16 @@ TSharedRef<SBox> SInterchangePipelineConfigurationDialog::SpawnPipelineConfigura
 		.bReimport(bReimport)
 		.PipelineStack(PipelineStack);
 
+	const float PipelineStackTreeViewWidth = GInterchangeEnableCustomPipelines ? 0.4f : 0.f;
+
 	TSharedPtr<SBox> InspectorBox;
 	TSharedRef<SBox> PipelineConfigurationPanelBox = SNew(SBox)
 	[
 		SNew(SSplitter)
 		.Orientation(Orient_Horizontal)
 		+ SSplitter::Slot()
-		.Value(0.4f)
+		.Value(PipelineStackTreeViewWidth)
+		.Resizable(GInterchangeEnableCustomPipelines)
 		[
 			SNew(SVerticalBox)
 			+ SVerticalBox::Slot()
@@ -451,7 +465,7 @@ TSharedRef<SBox> SInterchangePipelineConfigurationDialog::SpawnPipelineConfigura
 			]
 		]
 		+ SSplitter::Slot()
-		.Value(0.6f)
+		.Value(1.f - PipelineStackTreeViewWidth)
 		[
 			SAssignNew(InspectorBox, SBox)
 		]
@@ -584,6 +598,8 @@ void SInterchangePipelineConfigurationDialog::Construct(const FArguments& InArgs
 			]
 		]
 	];
+
+	PipelineConfigurationTreeView->SelectDefaultItem();
 }
 
 void SInterchangePipelineConfigurationDialog::OnSelectionChanged(TSharedPtr<FInterchangePipelineStacksTreeNodeItem> Item, ESelectInfo::Type SelectionType)
diff --git a/Engine/Plugins/Experimental/Interchange/Editor/Source/Pipelines/Private/SInterchangePipelineConfigurationDialog.h b/Engine/Plugins/Experimental/Interchange/Editor/Source/Pipelines/Private/SInterchangePipelineConfigurationDialog.h
--- a/Engine/Plugins/Experimental/Interchange/Editor/Source/Pipelines/Private/SInterchangePipelineConfigurationDialog.h
+++ b/Engine/Plugins/Experimental/Interchange/Editor/Source/Pipelines/Private/SInterchangePipelineConfigurationDialog.h
@@ -68,6 +68,9 @@ public:
 
 	const TArray<TSharedPtr<FInterchangePipelineStacksTreeNodeItem>>& GetRootNodeArray() const { return RootNodeArray; }
 	TArray<TSharedPtr<FInterchangePipelineStacksTreeNodeItem>>& GetMutableRootNodeArray() { return RootNodeArray; }
+
+	void SelectDefaultItem();
+
 protected:
 	/** Delegate to invoke when selection changes. */
 	FOnPipelineConfigurationSelectionChanged OnSelectionChangedDelegate;
diff --git a/Engine/Source/Runtime/Experimental/Interchange/Core/Private/InterchangePipelineBase.cpp b/Engine/Source/Runtime/Experimental/Interchange/Core/Private/InterchangePipelineBase.cpp
--- a/Engine/Source/Runtime/Experimental/Interchange/Core/Private/InterchangePipelineBase.cpp
+++ b/Engine/Source/Runtime/Experimental/Interchange/Core/Private/InterchangePipelineBase.cpp
@@ -2,12 +2,20 @@
 
 #include "InterchangePipelineBase.h"
 
-#include "CoreMinimal.h"
+#include "HAL/IConsoleManager.h"
 #include "InterchangeLogPrivate.h"
 #include "Misc/ConfigCacheIni.h"
 #include "UObject/Object.h"
 #include "UObject/ObjectMacros.h"
 
+INTERCHANGECORE_API bool GInterchangeEnableCustomPipelines = true;
+
+static FAutoConsoleVariableRef CCvarInterchangeEnableCustomPipelines(
+	TEXT("Interchange.FeatureFlags.CustomPipelines"),
+	GInterchangeEnableCustomPipelines,
+	TEXT("Whether custom pipelines support is enabled."),
+	ECVF_Default);
+
 namespace UE
 {
 	namespace Interchange
