Interchange
- Allow different pipelines per asset type so that we don't show all the asset options when importing simple assets likes textures and materials.
- Allow to choose to display the config dialog per asset type so that we can skip it for things like textures.

- Added default stack overrides per asset type.
- Added show config dialog overrides per asset type.
- Added a DefaultTexturePipeline asset which is the Generic Texture Pipeline with the "Import Textures" option hidden.
- Added a DefaultMaterialPipeline asset which is the Generic Material Pipeline with the "Import Materials" option hidden.
- Generic Material Pipeline: Split the import materials option from the type of material to import as so that we can configure their visibility independently.
- Moved the generic texture pipeline from the generic asset pipeline to the generic material pipeline so that we can use the generic material pipeline on its own.

diff --git a/Engine/Config/BaseEngine.ini b/Engine/Config/BaseEngine.ini
--- a/Engine/Config/BaseEngine.ini
+++ b/Engine/Config/BaseEngine.ini
@@ -3081,5 +3081,5 @@ EnvironmentCollisionMaxShapes = 32
 [/Script/InterchangeEngine.InterchangeProjectSettings]
 FilePickerClass=/Script/InterchangeEditorUtilities.InterchangeFilePickerGeneric
 bStaticMeshUseSmoothEdgesIfSmoothingInformationIsMissing=True
-ContentImportSettings=(PipelineStacks=(("Assets", (Pipelines=("/Interchange/Pipelines/DefaultAssetsPipeline.DefaultAssetsPipeline")))),DefaultPipelineStack="Assets",PipelineConfigurationDialogClass="/Script/InterchangeEditorPipelines.InterchangePipelineConfigurationGeneric",bShowPipelineStacksConfigurationDialog=True)
+ContentImportSettings=(DefaultPipelineStackOverride=((Materials, "Materials"),(Textures, "Textures")),ShowPipelineStacksConfigurationDialogOverride=((Textures, False)),PipelineStacks=(("Assets", (Pipelines=("/Interchange/Pipelines/DefaultAssetsPipeline.DefaultAssetsPipeline"))),("Materials", (Pipelines=("/Interchange/Pipelines/DefaultMaterialPipeline.DefaultMaterialPipeline"))),("Textures", (Pipelines=("/Interchange/Pipelines/DefaultTexturePipeline.DefaultTexturePipeline")))),DefaultPipelineStack="Assets",PipelineConfigurationDialogClass="/Script/InterchangeEditorPipelines.InterchangePipelineConfigurationGeneric",bShowPipelineStacksConfigurationDialog=True)
 SceneImportSettings=(PipelineStacks=(("Scene", (Pipelines=("/Interchange/Pipelines/DefaultSceneAssetsPipeline.DefaultSceneAssetsPipeline","/Interchange/Pipelines/DefaultSceneLevelPipeline.DefaultSceneLevelPipeline")))),DefaultPipelineStack="Scene",PipelineConfigurationDialogClass="/Script/InterchangeEditorPipelines.InterchangePipelineConfigurationGeneric",bShowPipelineStacksConfigurationDialog=True)
diff --git a/Engine/Plugins/Experimental/Interchange/Editor/Source/Pipelines/Private/SInterchangePipelineConfigurationDialog.cpp b/Engine/Plugins/Experimental/Interchange/Editor/Source/Pipelines/Private/SInterchangePipelineConfigurationDialog.cpp
--- a/Engine/Plugins/Experimental/Interchange/Editor/Source/Pipelines/Private/SInterchangePipelineConfigurationDialog.cpp
+++ b/Engine/Plugins/Experimental/Interchange/Editor/Source/Pipelines/Private/SInterchangePipelineConfigurationDialog.cpp
@@ -26,30 +26,12 @@
 #include "Widgets/Layout/SSplitter.h"
 #include "Widgets/Layout/SUniformGridPanel.h"
 
+#include "ObjectEditorUtils.h"
+
 #define LOCTEXT_NAMESPACE "InterchangePipelineConfiguration"
 
 extern INTERCHANGECORE_API bool GInterchangeEnableCustomPipelines;
 
-namespace UE::Interchange::Private
-{
-	const FInterchangeImportSettings& GetDefaultImportSettings(const bool bIsSceneImport)
-	{
-		if (bIsSceneImport)
-		{
-			return GetDefault<UInterchangeProjectSettings>()->SceneImportSettings;
-		}
-		else
-		{
-			return GetDefault<UInterchangeProjectSettings>()->ContentImportSettings;
-		}
-	}
-
-	FName GetDefaultPipelineStackName(const bool bIsSceneImport)
-	{
-		return GetDefaultImportSettings(bIsSceneImport).DefaultPipelineStack;
-	}
-}
-
 /************************************************************************/
 /* FInterchangePipelineStacksTreeNodeItem Implementation                    */
 /************************************************************************/
@@ -74,26 +56,15 @@ SInterchangePipelineStacksTreeView::~SInterchangePipelineStacksTreeView()
 void SInterchangePipelineStacksTreeView::Construct(const FArguments& InArgs)
 {
 	OnSelectionChangedDelegate = InArgs._OnSelectionChangedDelegate;
+	SourceData = InArgs._SourceData;
 	PipelineStack = InArgs._PipelineStack;
 	bSceneImport = InArgs._bSceneImport;
 	bReimport = InArgs._bReimport && PipelineStack.Num() > 0;
 
 	//Build the FbxNodeInfoPtr tree data
 	const FName ReimportPipelineName = TEXT("ReimportPipeline");
-	const FName& DefaultPipelineStackName = bReimport ? ReimportPipelineName : UE::Interchange::Private::GetDefaultPipelineStackName(bSceneImport);
-	const TMap<FName, FInterchangePipelineStack>& PipelineStacks =
-		[this]() -> const TMap<FName, FInterchangePipelineStack>&
-		{
-			if (bSceneImport)
-			{
-				return GetDefault<UInterchangeProjectSettings>()->SceneImportSettings.PipelineStacks;
-			}
-			else
-			{
-				return GetDefault<UInterchangeProjectSettings>()->ContentImportSettings.PipelineStacks;
-			}
-		}();
-		
+	const FName& DefaultPipelineStackName = bReimport ? ReimportPipelineName : FInterchangeProjectSettingsUtils::GetDefaultPipelineStackName(bSceneImport, *SourceData.Get());
+	const TMap<FName, FInterchangePipelineStack>& PipelineStacks = FInterchangeProjectSettingsUtils::GetDefaultImportSettings(bSceneImport).PipelineStacks;
 
 	//In reimport we modify directly the asset pipeline
 	if (bReimport)
@@ -156,10 +127,18 @@ void SInterchangePipelineStacksTreeView::Construct(const FArguments& InArgs)
 
 void SInterchangePipelineStacksTreeView::SelectDefaultItem()
 {
-	if (!RootNodeArray.IsEmpty() && !RootNodeArray[0]->Childrens.IsEmpty())
+	FName DefaultStackName = FInterchangeProjectSettingsUtils::GetDefaultPipelineStackName(bSceneImport, *SourceData.Get());
+
+	for (const TSharedPtr<FInterchangePipelineStacksTreeNodeItem>& PipelineStacksTreeNodeItem : RootNodeArray)
 	{
-		SetItemExpansion(RootNodeArray[0], true);
-		SetSelection(RootNodeArray[0]->Childrens[0]);
+		if (PipelineStacksTreeNodeItem->StackName == DefaultStackName && !PipelineStacksTreeNodeItem->Childrens.IsEmpty())
+		{
+			const bool bShouldExpandItem = true;
+			SetItemExpansion(PipelineStacksTreeNodeItem, bShouldExpandItem);
+			SetSelection(PipelineStacksTreeNodeItem->Childrens[0]);
+
+			break;
+		}
 	}
 }
 
@@ -175,6 +154,7 @@ public:
 		/** The item content. */
 		SLATE_ARGUMENT(TSharedPtr<FInterchangePipelineStacksTreeNodeItem>, InterchangeNode)
 		SLATE_ARGUMENT(bool, bSceneImport)
+		SLATE_ARGUMENT(FName, DefaultPipelineStackName)
 	SLATE_END_ARGS()
 
 	/**
@@ -186,6 +166,7 @@ public:
 	{
 		InterchangeNode = InArgs._InterchangeNode;
 		bSceneImport = InArgs._bSceneImport;
+		DefaultPipelineStackName = InArgs._DefaultPipelineStackName;
 
 		//This is suppose to always be valid
 		check(InterchangeNode);
@@ -231,7 +212,6 @@ public:
 private:
 	const FSlateBrush* GetImageItemIcon() const
 	{
-		const FName DefaultPipelineStackName = UE::Interchange::Private::GetDefaultPipelineStackName(bSceneImport);
 		//This is suppose to always be valid
 		check(InterchangeNode);
 		const bool bIsPipelineStackNode = InterchangeNode->Pipeline == nullptr;
@@ -254,13 +234,15 @@ private:
 	/** The node to build the tree view row from. */
 	TSharedPtr<FInterchangePipelineStacksTreeNodeItem> InterchangeNode = nullptr;
 	bool bSceneImport = false;
+	FName DefaultPipelineStackName;
 };
 
 TSharedRef< ITableRow > SInterchangePipelineStacksTreeView::OnGenerateRowPipelineConfigurationTreeView(TSharedPtr<FInterchangePipelineStacksTreeNodeItem> Item, const TSharedRef< STableViewBase >& OwnerTable)
 {
 	TSharedRef< SInterchangePipelineStacksTreeViewItem > ReturnRow = SNew(SInterchangePipelineStacksTreeViewItem, OwnerTable)
 		.InterchangeNode(Item)
-		.bSceneImport(bSceneImport);
+		.bSceneImport(bSceneImport)
+		.DefaultPipelineStackName(FInterchangeProjectSettingsUtils::GetDefaultPipelineStackName(bSceneImport, *SourceData.Get()));
 	return ReturnRow;
 }
 void SInterchangePipelineStacksTreeView::OnGetChildrenPipelineConfigurationTreeView(TSharedPtr<FInterchangePipelineStacksTreeNodeItem> InParent, TArray< TSharedPtr<FInterchangePipelineStacksTreeNodeItem> >& OutChildren)
@@ -418,6 +400,7 @@ TSharedRef<SBox> SInterchangePipelineConfigurationDialog::SpawnPipelineConfigura
 	//Create the treeview
 	PipelineConfigurationTreeView = SNew(SInterchangePipelineStacksTreeView)
 		.OnSelectionChangedDelegate(this, &SInterchangePipelineConfigurationDialog::OnSelectionChanged)
+		.SourceData(SourceData)
 		.bSceneImport(bSceneImport)
 		.bReimport(bReimport)
 		.PipelineStack(PipelineStack);
@@ -684,7 +667,7 @@ FReply SInterchangePipelineConfigurationDialog::OnResetToDefault()
 			else
 			{
 				//The reset to default when doing a import is to duplicate again the pipeline stack reference asset.
-				const TMap<FName, FInterchangePipelineStack>& DefaultPipelineStacks = UE::Interchange::Private::GetDefaultImportSettings(bSceneImport).PipelineStacks;
+				const TMap<FName, FInterchangePipelineStack>& DefaultPipelineStacks = FInterchangeProjectSettingsUtils::GetDefaultImportSettings(bSceneImport).PipelineStacks;
 				//Find the stack node, so we can update the pipeline
 				TSharedPtr<FInterchangePipelineStacksTreeNodeItem> StackNode;
 				TArray<TSharedPtr<FInterchangePipelineStacksTreeNodeItem>>& RootNodeArray = PipelineConfigurationTreeView->GetMutableRootNodeArray();
@@ -748,7 +731,7 @@ bool SInterchangePipelineConfigurationDialog::RecursiveValidatePipelineSettings(
 
 bool SInterchangePipelineConfigurationDialog::IsImportButtonEnabled() const
 {
-	const FName DefaultPipelineStackName = UE::Interchange::Private::GetDefaultPipelineStackName(bSceneImport);
+	const FName DefaultPipelineStackName = FInterchangeProjectSettingsUtils::GetDefaultPipelineStackName(bSceneImport, *SourceData.Get());
 	const TArray<TSharedPtr<FInterchangePipelineStacksTreeNodeItem>>& RootNodeArray = PipelineConfigurationTreeView->GetRootNodeArray();
 	for (const TSharedPtr<FInterchangePipelineStacksTreeNodeItem>& RootNode : RootNodeArray)
 	{
@@ -764,7 +747,7 @@ bool SInterchangePipelineConfigurationDialog::IsImportButtonEnabled() const
 
 FText SInterchangePipelineConfigurationDialog::GetImportButtonTooltip() const
 {
-	const FName DefaultPipelineStackName = UE::Interchange::Private::GetDefaultPipelineStackName(bSceneImport);
+	const FName DefaultPipelineStackName = FInterchangeProjectSettingsUtils::GetDefaultPipelineStackName(bSceneImport, *SourceData.Get());
 	const TArray<TSharedPtr<FInterchangePipelineStacksTreeNodeItem>>& RootNodeArray = PipelineConfigurationTreeView->GetRootNodeArray();
 	for (const TSharedPtr<FInterchangePipelineStacksTreeNodeItem>& RootNode : RootNodeArray)
 	{
diff --git a/Engine/Plugins/Experimental/Interchange/Editor/Source/Pipelines/Private/SInterchangePipelineConfigurationDialog.h b/Engine/Plugins/Experimental/Interchange/Editor/Source/Pipelines/Private/SInterchangePipelineConfigurationDialog.h
--- a/Engine/Plugins/Experimental/Interchange/Editor/Source/Pipelines/Private/SInterchangePipelineConfigurationDialog.h
+++ b/Engine/Plugins/Experimental/Interchange/Editor/Source/Pipelines/Private/SInterchangePipelineConfigurationDialog.h
@@ -53,6 +53,7 @@ public:
 		: _OnSelectionChangedDelegate()
 	{}
 		SLATE_EVENT(FOnPipelineConfigurationSelectionChanged, OnSelectionChangedDelegate)
+		SLATE_ARGUMENT(TWeakObjectPtr<UInterchangeSourceData>, SourceData)
 		SLATE_ARGUMENT(bool, bSceneImport)
 		SLATE_ARGUMENT(bool, bReimport)
 		SLATE_ARGUMENT(TArray<UInterchangePipelineBase*>, PipelineStack)
@@ -74,6 +75,7 @@ public:
 protected:
 	/** Delegate to invoke when selection changes. */
 	FOnPipelineConfigurationSelectionChanged OnSelectionChangedDelegate;
+	TWeakObjectPtr<UInterchangeSourceData> SourceData;
 	bool bSceneImport = false;
 	bool bReimport = false;
 	TArray<UInterchangePipelineBase*> PipelineStack;
diff --git a/Engine/Plugins/Experimental/Interchange/Runtime/Source/Import/Private/Fbx/InterchangeFbxTranslator.cpp b/Engine/Plugins/Experimental/Interchange/Runtime/Source/Import/Private/Fbx/InterchangeFbxTranslator.cpp
--- a/Engine/Plugins/Experimental/Interchange/Runtime/Source/Import/Private/Fbx/InterchangeFbxTranslator.cpp
+++ b/Engine/Plugins/Experimental/Interchange/Runtime/Source/Import/Private/Fbx/InterchangeFbxTranslator.cpp
@@ -40,15 +40,10 @@ EInterchangeTranslatorType UInterchangeFbxTranslator::GetTranslatorType() const
 	return EInterchangeTranslatorType::Scenes;
 }
 
-bool UInterchangeFbxTranslator::DoesSupportAssetType(EInterchangeTranslatorAssetType AssetType) const
+EInterchangeTranslatorAssetType UInterchangeFbxTranslator::GetSupportedAssetTypes() const
 {
 	//fbx translator support Meshes, Materials and animation
-	if (AssetType == EInterchangeTranslatorAssetType::Materials || AssetType == EInterchangeTranslatorAssetType::Meshes || AssetType == EInterchangeTranslatorAssetType::Animations)
-	{
-		return true;
-	}
-
-	return false;
+	return EInterchangeTranslatorAssetType::Materials | EInterchangeTranslatorAssetType::Meshes | EInterchangeTranslatorAssetType::Animations;
 }
 
 TArray<FString> UInterchangeFbxTranslator::GetSupportedFormats() const
diff --git a/Engine/Plugins/Experimental/Interchange/Runtime/Source/Import/Private/Gltf/InterchangeGltfTranslator.cpp b/Engine/Plugins/Experimental/Interchange/Runtime/Source/Import/Private/Gltf/InterchangeGltfTranslator.cpp
--- a/Engine/Plugins/Experimental/Interchange/Runtime/Source/Import/Private/Gltf/InterchangeGltfTranslator.cpp
+++ b/Engine/Plugins/Experimental/Interchange/Runtime/Source/Import/Private/Gltf/InterchangeGltfTranslator.cpp
@@ -791,15 +791,10 @@ EInterchangeTranslatorType UInterchangeGltfTranslator::GetTranslatorType() const
 	return EInterchangeTranslatorType::Scenes;
 }
 
-bool UInterchangeGltfTranslator::DoesSupportAssetType(EInterchangeTranslatorAssetType AssetType) const
+EInterchangeTranslatorAssetType UInterchangeGltfTranslator::GetSupportedAssetTypes() const
 {
 	//gltf translator support Meshes and Materials
-	if (AssetType == EInterchangeTranslatorAssetType::Materials || AssetType == EInterchangeTranslatorAssetType::Meshes || AssetType == EInterchangeTranslatorAssetType::Animations)
-	{
-		return true;
-	}
-
-	return false;
+	return EInterchangeTranslatorAssetType::Materials | EInterchangeTranslatorAssetType::Meshes | EInterchangeTranslatorAssetType::Animations;
 }
 
 TArray<FString> UInterchangeGltfTranslator::GetSupportedFormats() const
diff --git a/Engine/Plugins/Experimental/Interchange/Runtime/Source/Import/Private/MaterialX/InterchangeMaterialXTranslator.cpp b/Engine/Plugins/Experimental/Interchange/Runtime/Source/Import/Private/MaterialX/InterchangeMaterialXTranslator.cpp
--- a/Engine/Plugins/Experimental/Interchange/Runtime/Source/Import/Private/MaterialX/InterchangeMaterialXTranslator.cpp
+++ b/Engine/Plugins/Experimental/Interchange/Runtime/Source/Import/Private/MaterialX/InterchangeMaterialXTranslator.cpp
@@ -149,9 +149,9 @@ EInterchangeTranslatorType UInterchangeMaterialXTranslator::GetTranslatorType()
 	return EInterchangeTranslatorType::Scenes;
 }
 
-bool UInterchangeMaterialXTranslator::DoesSupportAssetType(EInterchangeTranslatorAssetType AssetType) const
+EInterchangeTranslatorAssetType UInterchangeMaterialXTranslator::GetSupportedAssetTypes() const
 {
-	return AssetType == EInterchangeTranslatorAssetType::Materials;
+	return EInterchangeTranslatorAssetType::Materials;
 }
 
 TArray<FString> UInterchangeMaterialXTranslator::GetSupportedFormats() const
diff --git a/Engine/Plugins/Experimental/Interchange/Runtime/Source/Import/Private/Mesh/InterchangeOBJTranslator.cpp b/Engine/Plugins/Experimental/Interchange/Runtime/Source/Import/Private/Mesh/InterchangeOBJTranslator.cpp
--- a/Engine/Plugins/Experimental/Interchange/Runtime/Source/Import/Private/Mesh/InterchangeOBJTranslator.cpp
+++ b/Engine/Plugins/Experimental/Interchange/Runtime/Source/Import/Private/Mesh/InterchangeOBJTranslator.cpp
@@ -871,15 +871,10 @@ TArray<FString> UInterchangeOBJTranslator::GetSupportedFormats() const
 	}
 }
 
-bool UInterchangeOBJTranslator::DoesSupportAssetType(EInterchangeTranslatorAssetType AssetType) const
+EInterchangeTranslatorAssetType UInterchangeOBJTranslator::GetSupportedAssetTypes() const
 {
 	//Obj translator support Meshes and Materials
-	if (AssetType == EInterchangeTranslatorAssetType::Materials || AssetType == EInterchangeTranslatorAssetType::Meshes)
-	{
-		return true;
-	}
-
-	return false;
+	return EInterchangeTranslatorAssetType::Materials | EInterchangeTranslatorAssetType::Meshes;
 }
 
 
diff --git a/Engine/Plugins/Experimental/Interchange/Runtime/Source/Import/Public/Fbx/InterchangeFbxTranslator.h b/Engine/Plugins/Experimental/Interchange/Runtime/Source/Import/Public/Fbx/InterchangeFbxTranslator.h
--- a/Engine/Plugins/Experimental/Interchange/Runtime/Source/Import/Public/Fbx/InterchangeFbxTranslator.h
+++ b/Engine/Plugins/Experimental/Interchange/Runtime/Source/Import/Public/Fbx/InterchangeFbxTranslator.h
@@ -31,7 +31,7 @@ public:
 
 	/** Begin UInterchangeTranslatorBase API*/
 	virtual EInterchangeTranslatorType GetTranslatorType() const override;
-	virtual bool DoesSupportAssetType(EInterchangeTranslatorAssetType AssetType) const override;
+	virtual EInterchangeTranslatorAssetType GetSupportedAssetTypes() const override;
 	virtual TArray<FString> GetSupportedFormats() const override;
 	virtual bool Translate(UInterchangeBaseNodeContainer& BaseNodeContainer) const override;
 	virtual void ReleaseSource() override;
diff --git a/Engine/Plugins/Experimental/Interchange/Runtime/Source/Import/Public/Gltf/InterchangeGltfTranslator.h b/Engine/Plugins/Experimental/Interchange/Runtime/Source/Import/Public/Gltf/InterchangeGltfTranslator.h
--- a/Engine/Plugins/Experimental/Interchange/Runtime/Source/Import/Public/Gltf/InterchangeGltfTranslator.h
+++ b/Engine/Plugins/Experimental/Interchange/Runtime/Source/Import/Public/Gltf/InterchangeGltfTranslator.h
@@ -36,7 +36,7 @@ class UInterchangeGltfTranslator : public UInterchangeTranslatorBase
 public:
 	/** Begin UInterchangeTranslatorBase API*/
 	virtual EInterchangeTranslatorType GetTranslatorType() const override;
-	virtual bool DoesSupportAssetType(EInterchangeTranslatorAssetType AssetType) const override;
+	virtual EInterchangeTranslatorAssetType GetSupportedAssetTypes() const override;
 	virtual TArray<FString> GetSupportedFormats() const override;
 	virtual bool Translate( UInterchangeBaseNodeContainer& BaseNodeContainer ) const override;
 	/** End UInterchangeTranslatorBase API*/
diff --git a/Engine/Plugins/Experimental/Interchange/Runtime/Source/Import/Public/MaterialX/InterchangeMaterialXTranslator.h b/Engine/Plugins/Experimental/Interchange/Runtime/Source/Import/Public/MaterialX/InterchangeMaterialXTranslator.h
--- a/Engine/Plugins/Experimental/Interchange/Runtime/Source/Import/Public/MaterialX/InterchangeMaterialXTranslator.h
+++ b/Engine/Plugins/Experimental/Interchange/Runtime/Source/Import/Public/MaterialX/InterchangeMaterialXTranslator.h
@@ -34,7 +34,7 @@ public:
 
 	virtual EInterchangeTranslatorType GetTranslatorType() const override;
 
-	virtual bool DoesSupportAssetType(EInterchangeTranslatorAssetType AssetType) const override;
+	virtual EInterchangeTranslatorAssetType GetSupportedAssetTypes() const override;
 
 	virtual TArray<FString> GetSupportedFormats() const override;
 
diff --git a/Engine/Plugins/Experimental/Interchange/Runtime/Source/Import/Public/Mesh/InterchangeOBJTranslator.h b/Engine/Plugins/Experimental/Interchange/Runtime/Source/Import/Public/Mesh/InterchangeOBJTranslator.h
--- a/Engine/Plugins/Experimental/Interchange/Runtime/Source/Import/Public/Mesh/InterchangeOBJTranslator.h
+++ b/Engine/Plugins/Experimental/Interchange/Runtime/Source/Import/Public/Mesh/InterchangeOBJTranslator.h
@@ -31,7 +31,7 @@ public:
 
 
 	virtual TArray<FString> GetSupportedFormats() const override;
-	virtual bool DoesSupportAssetType(EInterchangeTranslatorAssetType AssetType) const override;
+	virtual EInterchangeTranslatorAssetType GetSupportedAssetTypes() const override;
 	/**
 	 * Translate the associated source data into a node hold by the specified nodes container.
 	 *
diff --git a/Engine/Plugins/Experimental/Interchange/Runtime/Source/Import/Public/Texture/InterchangeDDSTranslator.h b/Engine/Plugins/Experimental/Interchange/Runtime/Source/Import/Public/Texture/InterchangeDDSTranslator.h
--- a/Engine/Plugins/Experimental/Interchange/Runtime/Source/Import/Public/Texture/InterchangeDDSTranslator.h
+++ b/Engine/Plugins/Experimental/Interchange/Runtime/Source/Import/Public/Texture/InterchangeDDSTranslator.h
@@ -21,7 +21,7 @@ class INTERCHANGEIMPORT_API UInterchangeDDSTranslator : public UInterchangeTrans
 public:
 
 	virtual TArray<FString> GetSupportedFormats() const override;
-	virtual bool DoesSupportAssetType(EInterchangeTranslatorAssetType AssetType) const override { return AssetType == EInterchangeTranslatorAssetType::Textures; }
+	virtual EInterchangeTranslatorAssetType GetSupportedAssetTypes() const override { return EInterchangeTranslatorAssetType::Textures; }
 	
 	/*
 	 * return true if the translator can translate the specified source data.
diff --git a/Engine/Plugins/Experimental/Interchange/Runtime/Source/Import/Public/Texture/InterchangeIESTranslator.h b/Engine/Plugins/Experimental/Interchange/Runtime/Source/Import/Public/Texture/InterchangeIESTranslator.h
--- a/Engine/Plugins/Experimental/Interchange/Runtime/Source/Import/Public/Texture/InterchangeIESTranslator.h
+++ b/Engine/Plugins/Experimental/Interchange/Runtime/Source/Import/Public/Texture/InterchangeIESTranslator.h
@@ -19,7 +19,7 @@ class INTERCHANGEIMPORT_API UInterchangeIESTranslator : public UInterchangeTrans
 public:
 
 	virtual TArray<FString> GetSupportedFormats() const override;
-	virtual bool DoesSupportAssetType(EInterchangeTranslatorAssetType AssetType) const override { return AssetType == EInterchangeTranslatorAssetType::Textures; }
+	virtual EInterchangeTranslatorAssetType GetSupportedAssetTypes() const override { return EInterchangeTranslatorAssetType::Textures; }
 
 	/**
 	 * Translate the associated source data into a node hold by the specified nodes container.
diff --git a/Engine/Plugins/Experimental/Interchange/Runtime/Source/Import/Public/Texture/InterchangeImageWrapperTranslator.h b/Engine/Plugins/Experimental/Interchange/Runtime/Source/Import/Public/Texture/InterchangeImageWrapperTranslator.h
--- a/Engine/Plugins/Experimental/Interchange/Runtime/Source/Import/Public/Texture/InterchangeImageWrapperTranslator.h
+++ b/Engine/Plugins/Experimental/Interchange/Runtime/Source/Import/Public/Texture/InterchangeImageWrapperTranslator.h
@@ -23,7 +23,7 @@ public:
 
 	virtual TArray<FString> GetSupportedFormats() const override;
 
-	virtual bool DoesSupportAssetType(EInterchangeTranslatorAssetType AssetType) const override { return AssetType == EInterchangeTranslatorAssetType::Textures; }
+	virtual EInterchangeTranslatorAssetType GetSupportedAssetTypes() const override { return EInterchangeTranslatorAssetType::Textures; }
 
 	/**
 	 * Translate the associated source data into a node hold by the specified nodes container.
diff --git a/Engine/Plugins/Experimental/Interchange/Runtime/Source/Import/Public/Texture/InterchangeJPGTranslator.h b/Engine/Plugins/Experimental/Interchange/Runtime/Source/Import/Public/Texture/InterchangeJPGTranslator.h
--- a/Engine/Plugins/Experimental/Interchange/Runtime/Source/Import/Public/Texture/InterchangeJPGTranslator.h
+++ b/Engine/Plugins/Experimental/Interchange/Runtime/Source/Import/Public/Texture/InterchangeJPGTranslator.h
@@ -17,9 +17,9 @@ class INTERCHANGEIMPORT_API UInterchangeJPGTranslator : public UInterchangeTrans
 {
 	GENERATED_BODY()
 public:
-	
+
 	virtual TArray<FString> GetSupportedFormats() const override;
-	virtual bool DoesSupportAssetType(EInterchangeTranslatorAssetType AssetType) const override { return AssetType == EInterchangeTranslatorAssetType::Textures; }
+	virtual EInterchangeTranslatorAssetType GetSupportedAssetTypes() const override { return EInterchangeTranslatorAssetType::Textures; }
 
 	/**
 	 * Translate the associated source data into a node hold by the specified nodes container.
diff --git a/Engine/Plugins/Experimental/Interchange/Runtime/Source/Import/Public/Texture/InterchangePCXTranslator.h b/Engine/Plugins/Experimental/Interchange/Runtime/Source/Import/Public/Texture/InterchangePCXTranslator.h
--- a/Engine/Plugins/Experimental/Interchange/Runtime/Source/Import/Public/Texture/InterchangePCXTranslator.h
+++ b/Engine/Plugins/Experimental/Interchange/Runtime/Source/Import/Public/Texture/InterchangePCXTranslator.h
@@ -17,9 +17,9 @@ class INTERCHANGEIMPORT_API UInterchangePCXTranslator : public UInterchangeTrans
 {
 	GENERATED_BODY()
 public:
-	
+
 	virtual TArray<FString> GetSupportedFormats() const override;
-	virtual bool DoesSupportAssetType(EInterchangeTranslatorAssetType AssetType) const override { return AssetType == EInterchangeTranslatorAssetType::Textures; }
+	virtual EInterchangeTranslatorAssetType GetSupportedAssetTypes() const override { return EInterchangeTranslatorAssetType::Textures; }
 
 	/**
 	 * Translates the associated source data into a node hold by the specified nodes container.
diff --git a/Engine/Plugins/Experimental/Interchange/Runtime/Source/Import/Public/Texture/InterchangePSDTranslator.h b/Engine/Plugins/Experimental/Interchange/Runtime/Source/Import/Public/Texture/InterchangePSDTranslator.h
--- a/Engine/Plugins/Experimental/Interchange/Runtime/Source/Import/Public/Texture/InterchangePSDTranslator.h
+++ b/Engine/Plugins/Experimental/Interchange/Runtime/Source/Import/Public/Texture/InterchangePSDTranslator.h
@@ -19,7 +19,7 @@ class INTERCHANGEIMPORT_API UInterchangePSDTranslator : public UInterchangeTrans
 public:
 
 	virtual TArray<FString> GetSupportedFormats() const override;
-	virtual bool DoesSupportAssetType(EInterchangeTranslatorAssetType AssetType) const override { return AssetType == EInterchangeTranslatorAssetType::Textures; }
+	virtual EInterchangeTranslatorAssetType GetSupportedAssetTypes() const override { return EInterchangeTranslatorAssetType::Textures; }
 
 	/**
 	 * Translate the associated source data into a node hold by the specified nodes container.
diff --git a/Engine/Plugins/Experimental/Interchange/Runtime/Source/Pipelines/Private/InterchangeGenericAssetsPipeline.cpp b/Engine/Plugins/Experimental/Interchange/Runtime/Source/Pipelines/Private/InterchangeGenericAssetsPipeline.cpp
--- a/Engine/Plugins/Experimental/Interchange/Runtime/Source/Pipelines/Private/InterchangeGenericAssetsPipeline.cpp
+++ b/Engine/Plugins/Experimental/Interchange/Runtime/Source/Pipelines/Private/InterchangeGenericAssetsPipeline.cpp
@@ -32,7 +32,6 @@
 
 UInterchangeGenericAssetsPipeline::UInterchangeGenericAssetsPipeline()
 {
-	TexturePipeline = CreateDefaultSubobject<UInterchangeGenericTexturePipeline>("TexturePipeline");
 	MaterialPipeline = CreateDefaultSubobject<UInterchangeGenericMaterialPipeline>("MaterialPipeline");
 	CommonMeshesProperties = CreateDefaultSubobject<UInterchangeGenericCommonMeshesProperties>("CommonMeshesProperties");
 	CommonSkeletalMeshesAndAnimationsProperties = CreateDefaultSubobject<UInterchangeGenericCommonSkeletalMeshesAndAnimationsProperties>("CommonSkeletalMeshesAndAnimationsProperties");
@@ -49,11 +48,6 @@ void UInterchangeGenericAssetsPipeline::PreDialogCleanup(const FName PipelineSta
 	//We always clean the pipeline skeleton when showing the dialog
 	CommonSkeletalMeshesAndAnimationsProperties->Skeleton = nullptr;
 
-	if (TexturePipeline)
-	{
-		TexturePipeline->PreDialogCleanup(PipelineStackName);
-	}
-
 	if (MaterialPipeline)
 	{
 		MaterialPipeline->PreDialogCleanup(PipelineStackName);
@@ -74,11 +68,6 @@ void UInterchangeGenericAssetsPipeline::PreDialogCleanup(const FName PipelineSta
 
 bool UInterchangeGenericAssetsPipeline::IsSettingsAreValid(TOptional<FText>& OutInvalidReason) const
 {
-	if (TexturePipeline && !TexturePipeline->IsSettingsAreValid(OutInvalidReason))
-	{
-		return false;
-	}
-
 	if (MaterialPipeline && !MaterialPipeline->IsSettingsAreValid(OutInvalidReason))
 	{
 		return false;
@@ -111,11 +100,6 @@ void UInterchangeGenericAssetsPipeline::AdjustSettingsForContext(EInterchangePip
 {
 	Super::AdjustSettingsForContext(ImportType, ReimportAsset);
 
-	if (TexturePipeline)
-	{
-		TexturePipeline->AdjustSettingsForContext(ImportType, ReimportAsset);
-	}
-
 	if (MaterialPipeline)
 	{
 		MaterialPipeline->AdjustSettingsForContext(ImportType, ReimportAsset);
@@ -146,10 +130,6 @@ void UInterchangeGenericAssetsPipeline::ExecutePreImportPipeline(UInterchangeBas
 	//The parent Results container should be set at this point
 	ensure(!Results.IsNull());
 	{
-		if (TexturePipeline)
-		{
-			TexturePipeline->SetResultsContainer(Results);
-		}
 		if (MaterialPipeline)
 		{
 			MaterialPipeline->SetResultsContainer(Results);
@@ -170,11 +150,11 @@ void UInterchangeGenericAssetsPipeline::ExecutePreImportPipeline(UInterchangeBas
 	//When we import only animation we need to prevent material and physic asset to be created
 	if (CommonSkeletalMeshesAndAnimationsProperties->bImportOnlyAnimations)
 	{
-		MaterialPipeline->MaterialImport = EInterchangeMaterialImportOption::DoNotImport;
+		MaterialPipeline->bImportMaterials = false;
 		MeshPipeline->bImportStaticMeshes = false;
 		MeshPipeline->bCreatePhysicsAsset = false;
 		MeshPipeline->PhysicsAsset = nullptr;
-		TexturePipeline->bImportTextures = false;
+		MaterialPipeline->TexturePipeline->bImportTextures = false;
 	}
 
 	//////////////////////////////////////////////////////////////////////////
@@ -214,10 +194,6 @@ void UInterchangeGenericAssetsPipeline::ExecutePreImportPipeline(UInterchangeBas
 		SourceDatas.Add(SourceData);
 	}
 
-	if (TexturePipeline)
-	{
-		TexturePipeline->ScriptedExecutePreImportPipeline(InBaseNodeContainer, InSourceDatas);
-	}
 	if (MaterialPipeline)
 	{
 		MaterialPipeline->ScriptedExecutePreImportPipeline(InBaseNodeContainer, InSourceDatas);
@@ -241,10 +217,6 @@ void UInterchangeGenericAssetsPipeline::ExecutePreImportPipeline(UInterchangeBas
 
 void UInterchangeGenericAssetsPipeline::ExecutePostImportPipeline(const UInterchangeBaseNodeContainer* InBaseNodeContainer, const FString& NodeKey, UObject* CreatedAsset, bool bIsAReimport)
 {
-	if (TexturePipeline)
-	{
-		TexturePipeline->ScriptedExecutePostImportPipeline(InBaseNodeContainer, NodeKey, CreatedAsset, bIsAReimport);
-	}
 	if (MaterialPipeline)
 	{
 		MaterialPipeline->ScriptedExecutePostImportPipeline(InBaseNodeContainer, NodeKey, CreatedAsset, bIsAReimport);
@@ -276,7 +248,7 @@ void UInterchangeGenericAssetsPipeline::AddPackageMetaData(UObject* CreatedAsset
 	//Add UObject package meta data
 	if (UMetaData* MetaData = CreatedAsset->GetOutermost()->GetMetaData())
 	{
-		//Cleanup existing INTERCHANGE_ prefix metadata name for this object (in casse we re-import)
+		//Cleanup existing INTERCHANGE_ prefix metadata name for this object (in case we re-import)
 		{
 			TArray<FName> InterchangeMetaDataKeys;
 			if(TMap<FName, FString>* MetaDataMapPtr = MetaData->GetMapForObject(CreatedAsset))
@@ -494,11 +466,6 @@ void UInterchangeGenericAssetsPipeline::AddPackageMetaData(UObject* CreatedAsset
 
 void UInterchangeGenericAssetsPipeline::SetReimportSourceIndex(UClass* ReimportObjectClass, const int32 SourceFileIndex)
 {
-	if (TexturePipeline)
-	{
-		TexturePipeline->ScriptedSetReimportSourceIndex(ReimportObjectClass, SourceFileIndex);
-	}
-
 	if (MaterialPipeline)
 	{
 		MaterialPipeline->ScriptedSetReimportSourceIndex(ReimportObjectClass, SourceFileIndex);
diff --git a/Engine/Plugins/Experimental/Interchange/Runtime/Source/Pipelines/Private/InterchangeGenericMaterialPipeline.cpp b/Engine/Plugins/Experimental/Interchange/Runtime/Source/Pipelines/Private/InterchangeGenericMaterialPipeline.cpp
--- a/Engine/Plugins/Experimental/Interchange/Runtime/Source/Pipelines/Private/InterchangeGenericMaterialPipeline.cpp
+++ b/Engine/Plugins/Experimental/Interchange/Runtime/Source/Pipelines/Private/InterchangeGenericMaterialPipeline.cpp
@@ -4,6 +4,7 @@
 
 #include "CoreMinimal.h"
 
+#include "InterchangeGenericTexturePipeline.h"
 #include "InterchangeMaterialDefinitions.h"
 #include "InterchangeMaterialFactoryNode.h"
 #include "InterchangeShaderGraphNode.h"
@@ -62,10 +63,40 @@ FString LexToString(UInterchangeGenericMaterialPipeline::EMaterialInputType Valu
 	}
 }
 
+UInterchangeGenericMaterialPipeline::UInterchangeGenericMaterialPipeline()
+{
+	TexturePipeline = CreateDefaultSubobject<UInterchangeGenericTexturePipeline>("TexturePipeline");
+}
+
+void UInterchangeGenericMaterialPipeline::PreDialogCleanup(const FName PipelineStackName)
+{
+	if (TexturePipeline)
+	{
+		TexturePipeline->PreDialogCleanup(PipelineStackName);
+	}
+
+	SaveSettings(PipelineStackName);
+}
+
+bool UInterchangeGenericMaterialPipeline::IsSettingsAreValid(TOptional<FText>& OutInvalidReason) const
+{
+	if (TexturePipeline && !TexturePipeline->IsSettingsAreValid(OutInvalidReason))
+	{
+		return false;
+	}
+
+	return true;
+}
+
 void UInterchangeGenericMaterialPipeline::AdjustSettingsForContext(EInterchangePipelineContext ImportType, TObjectPtr<UObject> ReimportAsset)
 {
 	Super::AdjustSettingsForContext(ImportType, ReimportAsset);
 
+	if (TexturePipeline)
+	{
+		TexturePipeline->AdjustSettingsForContext(ImportType, ReimportAsset);
+	}
+
 	TArray<FString> HideCategories;
 	bool bIsObjectAMaterial = ReimportAsset.IsNull() ? false : ReimportAsset->IsA(UMaterialInterface::StaticClass());
 	if ((!bIsObjectAMaterial && ImportType == EInterchangePipelineContext::AssetReimport)
@@ -74,7 +105,7 @@ void UInterchangeGenericMaterialPipeline::AdjustSettingsForContext(EInterchangeP
 		|| ImportType == EInterchangePipelineContext::AssetAlternateSkinningImport
 		|| ImportType == EInterchangePipelineContext::AssetAlternateSkinningReimport)
 	{
-		MaterialImport = EInterchangeMaterialImportOption::DoNotImport;
+		bImportMaterials = false;
 		HideCategories.Add(TEXT("Materials"));
 	}
 
@@ -95,6 +126,16 @@ void UInterchangeGenericMaterialPipeline::ExecutePreImportPipeline(UInterchangeB
 		return;
 	}
 
+	//Set the result container to allow error message
+	//The parent Results container should be set at this point
+	ensure(!Results.IsNull());
+	{
+		if (TexturePipeline)
+		{
+			TexturePipeline->SetResultsContainer(Results);
+		}
+	}
+
 	BaseNodeContainer = InBaseNodeContainer;
 	SourceDatas.Empty(InSourceDatas.Num());
 	for (const UInterchangeSourceData* SourceData : InSourceDatas)
@@ -102,6 +143,11 @@ void UInterchangeGenericMaterialPipeline::ExecutePreImportPipeline(UInterchangeB
 		SourceDatas.Add(SourceData);
 	}
 	
+	if (TexturePipeline)
+	{
+		TexturePipeline->ScriptedExecutePreImportPipeline(InBaseNodeContainer, InSourceDatas);
+	}
+
 	//Find all translated node we need for this pipeline
 	BaseNodeContainer->IterateNodes([this](const FString& NodeUid, UInterchangeBaseNode* Node)
 	{
@@ -168,6 +214,22 @@ void UInterchangeGenericMaterialPipeline::ExecutePreImportPipeline(UInterchangeB
 	}
 }
 
+void UInterchangeGenericMaterialPipeline::ExecutePostImportPipeline(const UInterchangeBaseNodeContainer* InBaseNodeContainer, const FString& NodeKey, UObject* CreatedAsset, bool bIsAReimport)
+{
+	if (TexturePipeline)
+	{
+		TexturePipeline->ScriptedExecutePostImportPipeline(InBaseNodeContainer, NodeKey, CreatedAsset, bIsAReimport);
+	}
+}
+
+void UInterchangeGenericMaterialPipeline::SetReimportSourceIndex(UClass* ReimportObjectClass, const int32 SourceFileIndex)
+{
+	if (TexturePipeline)
+	{
+		TexturePipeline->ScriptedSetReimportSourceIndex(ReimportObjectClass, SourceFileIndex);
+	}
+}
+
 UInterchangeBaseMaterialFactoryNode* UInterchangeGenericMaterialPipeline::CreateBaseMaterialFactoryNode(const UInterchangeBaseNode* MaterialNode, TSubclassOf<UInterchangeBaseMaterialFactoryNode> NodeType)
 {
 	FString DisplayLabel = MaterialNode->GetDisplayLabel();
diff --git a/Engine/Plugins/Experimental/Interchange/Runtime/Source/Pipelines/Public/InterchangeGenericAssetsPipeline.h b/Engine/Plugins/Experimental/Interchange/Runtime/Source/Pipelines/Public/InterchangeGenericAssetsPipeline.h
--- a/Engine/Plugins/Experimental/Interchange/Runtime/Source/Pipelines/Public/InterchangeGenericAssetsPipeline.h
+++ b/Engine/Plugins/Experimental/Interchange/Runtime/Source/Pipelines/Public/InterchangeGenericAssetsPipeline.h
@@ -72,10 +72,6 @@ public:
 	UPROPERTY(VisibleAnywhere, Instanced, Category = "Materials")
 	TObjectPtr<UInterchangeGenericMaterialPipeline> MaterialPipeline;
 
-	//////	TEXTURES_CATEGORY Properties //////
-	UPROPERTY(VisibleAnywhere, Instanced, Category = "Textures")
-	TObjectPtr<UInterchangeGenericTexturePipeline> TexturePipeline;
-
 	virtual void PreDialogCleanup(const FName PipelineStackName) override;
 
 	virtual bool IsSettingsAreValid(TOptional<FText>& OutInvalidReason) const override;
diff --git a/Engine/Plugins/Experimental/Interchange/Runtime/Source/Pipelines/Public/InterchangeGenericMaterialPipeline.h b/Engine/Plugins/Experimental/Interchange/Runtime/Source/Pipelines/Public/InterchangeGenericMaterialPipeline.h
--- a/Engine/Plugins/Experimental/Interchange/Runtime/Source/Pipelines/Public/InterchangeGenericMaterialPipeline.h
+++ b/Engine/Plugins/Experimental/Interchange/Runtime/Source/Pipelines/Public/InterchangeGenericMaterialPipeline.h
@@ -13,6 +13,7 @@
 
 class UInterchangeBaseMaterialFactoryNode;
 class UInterchangeFunctionCallShaderNode;
+class UInterchangeGenericTexturePipeline;
 class UInterchangeShaderGraphNode;
 class UInterchangeShaderNode;
 class UInterchangeMaterialFactoryNode;
@@ -24,28 +25,42 @@ class UInterchangeResult;
 UENUM(BlueprintType)
 enum class EInterchangeMaterialImportOption : uint8
 {
-	DoNotImport,
 	ImportAsMaterials,
 	ImportAsMaterialInstances,
 };
 
-UCLASS(BlueprintType, hidedropdown, Experimental)
+UCLASS(BlueprintType, editinlinenew, Experimental)
 class INTERCHANGEPIPELINES_API UInterchangeGenericMaterialPipeline : public UInterchangePipelineBase
 {
 	GENERATED_BODY()
 
 public:
+	UInterchangeGenericMaterialPipeline();
+
+	/** If enabled, imports the texture assets found in the sources. */
 	UPROPERTY(EditAnywhere, BlueprintReadWrite, Category = "Materials")
+	bool bImportMaterials = true;
+
+	UPROPERTY(EditAnywhere, BlueprintReadWrite, Category = "Materials", Meta=(EditCondition="bImportMaterials"))
 	EInterchangeMaterialImportOption MaterialImport = EInterchangeMaterialImportOption::ImportAsMaterials;
 
 	/** Optional material used as the parent when importing materials as instances. If no parent material is specified, one will be automatically selected during the import process. */
-	UPROPERTY(EditAnywhere, BlueprintReadWrite, Category = "Materials", Meta= (EditCondition="MaterialImport==EInterchangeMaterialImportOption::ImportAsMaterialInstances", AllowedClasses="/Script/Engine.MaterialInterface"))
+	UPROPERTY(EditAnywhere, BlueprintReadWrite, Category = "Materials", Meta= (EditCondition="bImportMaterials && MaterialImport==EInterchangeMaterialImportOption::ImportAsMaterialInstances", AllowedClasses="/Script/Engine.MaterialInterface"))
 	FSoftObjectPath ParentMaterial;
 
+	UPROPERTY(VisibleAnywhere, Instanced, Category = "Textures")
+	TObjectPtr<UInterchangeGenericTexturePipeline> TexturePipeline;
+
+	/** BEGIN UInterchangePipelineBase overrides */
+	virtual void PreDialogCleanup(const FName PipelineStackName) override;
+	virtual bool IsSettingsAreValid(TOptional<FText>& OutInvalidReason) const override;
 	virtual void AdjustSettingsForContext(EInterchangePipelineContext ImportType, TObjectPtr<UObject> ReimportAsset) override;
 
 protected:
 	virtual void ExecutePreImportPipeline(UInterchangeBaseNodeContainer* InBaseNodeContainer, const TArray<UInterchangeSourceData*>& InSourceDatas) override;
+	virtual void ExecutePostImportPipeline(const UInterchangeBaseNodeContainer* BaseNodeContainer, const FString& NodeKey, UObject* CreatedAsset, bool bIsAReimport) override;
+	virtual void SetReimportSourceIndex(UClass* ReimportObjectClass, const int32 SourceFileIndex) override;
+	/** END UInterchangePipelineBase overrides */
 
 	UPROPERTY()
 	TObjectPtr<UInterchangeBaseNodeContainer> BaseNodeContainer;
diff --git a/Engine/Plugins/Experimental/Interchange/Runtime/Source/Pipelines/Public/InterchangeGenericTexturePipeline.h b/Engine/Plugins/Experimental/Interchange/Runtime/Source/Pipelines/Public/InterchangeGenericTexturePipeline.h
--- a/Engine/Plugins/Experimental/Interchange/Runtime/Source/Pipelines/Public/InterchangeGenericTexturePipeline.h
+++ b/Engine/Plugins/Experimental/Interchange/Runtime/Source/Pipelines/Public/InterchangeGenericTexturePipeline.h
@@ -13,7 +13,7 @@
 class UInterchangeTextureFactoryNode;
 class UInterchangeTextureNode;
 
-UCLASS(BlueprintType, hidedropdown, Experimental)
+UCLASS(BlueprintType, editinlinenew, Experimental)
 class INTERCHANGEPIPELINES_API UInterchangeGenericTexturePipeline : public UInterchangePipelineBase
 {
 	GENERATED_BODY()
@@ -28,19 +28,19 @@ public:
 	 * If enable, after a new import a test will be run to see if the texture is a normal map.
 	 * If the texture is a normal map the SRG, CompressionSettings and LODGroup settings will be adjusted.
 	 */
-	UPROPERTY(EditAnywhere, BlueprintReadWrite, Category = "Textures")
+	UPROPERTY(EditAnywhere, BlueprintReadWrite, Category = "Textures", Meta=(EditCondition="bImportTextures"))
 	bool bDetectNormalMapTexture = true;
 
 	/** If enabled, the texture's green channel will be inverted for normal maps. */
-	UPROPERTY(EditAnywhere, BlueprintReadWrite, Category = "Textures")
+	UPROPERTY(EditAnywhere, BlueprintReadWrite, Category = "Textures", Meta=(EditCondition="bImportTextures"))
 	bool bFlipNormalMapGreenChannel = false;
 
 	/** If enabled detect if a texture use a UDIM pattern and if so import it as UIDMs. */
-	UPROPERTY(EditAnywhere, BlueprintReadWrite, Category = "Textures")
+	UPROPERTY(EditAnywhere, BlueprintReadWrite, Category = "Textures", Meta=(EditCondition="bImportTextures"))
 	bool bImportUDIMs = true;
 
 	/** Specify the files type that should be imported as long/lat cubemap */
-	UPROPERTY(EditAnywhere, BlueprintReadWrite, Category = "Textures")
+	UPROPERTY(EditAnywhere, BlueprintReadWrite, Category = "Textures", Meta=(EditCondition="bImportTextures"))
 	TSet<FString> FileExtensionsToImportAsLongLatCubemap = {"hdr"};
 
 	/** 
@@ -50,13 +50,13 @@ public:
 	 * 
 	 * Compressed source source data is generally store the data as it is in the source file.
 	 */
-	UPROPERTY(EditAnywhere, BlueprintReadWrite, AdvancedDisplay, Category = "Textures")
+	UPROPERTY(EditAnywhere, BlueprintReadWrite, AdvancedDisplay, Category = "Textures", Meta=(EditCondition="bImportTextures"))
 	bool bPreferCompressedSourceData = false;
 
 #endif
 
 	/** Should the textures that have a non-power of two resolution be imported */
-	UPROPERTY(EditAnywhere, BlueprintReadWrite, AdvancedDisplay, Category = "Textures")
+	UPROPERTY(EditAnywhere, BlueprintReadWrite, AdvancedDisplay, Category = "Textures", Meta=(EditCondition="bImportTextures"))
 	bool bAllowNonPowerOfTwo = false;
 
 public:
diff --git a/Engine/Source/Runtime/Experimental/Interchange/Core/Public/InterchangeTranslatorBase.h b/Engine/Source/Runtime/Experimental/Interchange/Core/Public/InterchangeTranslatorBase.h
--- a/Engine/Source/Runtime/Experimental/Interchange/Core/Public/InterchangeTranslatorBase.h
+++ b/Engine/Source/Runtime/Experimental/Interchange/Core/Public/InterchangeTranslatorBase.h
@@ -32,12 +32,13 @@ enum class EInterchangeTranslatorType : uint8
 UENUM(BlueprintType)
 enum class EInterchangeTranslatorAssetType : uint8
 {
-	None = 0,
-	Textures = 1,
-	Materials = 2,
-	Meshes = 3,
-	Animations = 4,
+	None = 0x0,
+	Textures = 0x1 << 0,
+	Materials = 0x1 << 1,
+	Meshes = 0x1 << 2,
+	Animations = 0x1 << 3,
 };
+ENUM_CLASS_FLAGS(EInterchangeTranslatorAssetType);
 
 UCLASS(BlueprintType, Blueprintable, Abstract, Experimental)
 class INTERCHANGECORE_API UInterchangeTranslatorBase : public UObject
@@ -50,7 +51,9 @@ public:
 
 	/** Specifies the capabilities of a translator. */
 	virtual EInterchangeTranslatorType GetTranslatorType() const { return EInterchangeTranslatorType::Assets; }
-	virtual bool DoesSupportAssetType(EInterchangeTranslatorAssetType AssetType) const { return false; }
+
+	bool DoesSupportAssetType(EInterchangeTranslatorAssetType AssetType) const { return EnumHasAllFlags(GetSupportedAssetTypes(), AssetType); }
+	virtual EInterchangeTranslatorAssetType GetSupportedAssetTypes() const PURE_VIRTUAL(UInterchangeTranslatorBase::GetSupportedAssetTypes, return EInterchangeTranslatorAssetType::None;);
 
 	/** List of formats supported by the translator. Each entry is of the form "ext;Description" where ext is the file extension. */
 	virtual TArray<FString> GetSupportedFormats() const PURE_VIRTUAL(UInterchangeTranslatorBase::GetSupportedExtensions, return TArray<FString>(););
diff --git a/Engine/Source/Runtime/Experimental/Interchange/Engine/Private/InterchangeManager.cpp b/Engine/Source/Runtime/Experimental/Interchange/Engine/Private/InterchangeManager.cpp
--- a/Engine/Source/Runtime/Experimental/Interchange/Engine/Private/InterchangeManager.cpp
+++ b/Engine/Source/Runtime/Experimental/Interchange/Engine/Private/InterchangeManager.cpp
@@ -49,27 +49,9 @@ namespace UE::Interchange::Private
 #endif
 	}
 
-	const FInterchangeImportSettings& GetImportSettings(const UInterchangeProjectSettings& InterchangeProjectSettings, EImportType ImportType)
-	{
-		if (ImportType == EImportType::ImportType_Asset)
-		{
-			return InterchangeProjectSettings.ContentImportSettings;
-		}
-		else if (ImportType == EImportType::ImportType_Scene)
-		{
-			return InterchangeProjectSettings.SceneImportSettings;
-		}
-		else
-		{
-			ensure(false);
-			static FInterchangeImportSettings InvalidSettings;
-			return InvalidSettings;
-		}
-	}
-
 	FInterchangeImportSettings& GetMutableImportSettings(UInterchangeProjectSettings& InterchangeProjectSettings, EImportType ImportType)
 	{
-		return const_cast<FInterchangeImportSettings&>(GetImportSettings(InterchangeProjectSettings, ImportType));
+		return const_cast<FInterchangeImportSettings&>(FInterchangeProjectSettingsUtils::GetImportSettings(InterchangeProjectSettings, ImportType == EImportType::ImportType_Scene));
 	}
 }
 
@@ -731,8 +713,7 @@ UInterchangeManager::ImportInternal(const FString& ContentPath, const UInterchan
 	}
 
 	const bool bImportScene = ImportType == UE::Interchange::EImportType::ImportType_Scene;
-	const UInterchangeProjectSettings* InterchangeProjectSettings = GetDefault<UInterchangeProjectSettings>();
-	const FInterchangeImportSettings& InterchangeImportSettings = UE::Interchange::Private::GetImportSettings(*InterchangeProjectSettings, ImportType);
+	const FInterchangeImportSettings& InterchangeImportSettings = FInterchangeProjectSettingsUtils::GetDefaultImportSettings(bImportScene);
 	
 	if (InterchangeImportSettings.PipelineStacks.Num() == 0)
 	{
@@ -743,8 +724,7 @@ UInterchangeManager::ImportInternal(const FString& ContentPath, const UInterchan
 	//Set a default pipeline stack if none is valid
 	if (!InterchangeImportSettings.PipelineStacks.Contains(InterchangeImportSettings.DefaultPipelineStack))
 	{
-		UInterchangeProjectSettings* MutableInterchangeProjectSettings = GetMutableDefault<UInterchangeProjectSettings>();
-		FInterchangeImportSettings& MutableInterchangeImportSettings = UE::Interchange::Private::GetMutableImportSettings(*MutableInterchangeProjectSettings, ImportType);
+		FInterchangeImportSettings& MutableInterchangeImportSettings = FInterchangeProjectSettingsUtils::GetMutableDefaultImportSettings(bImportScene);
 		TArray<FName> Keys;
 		MutableInterchangeImportSettings.PipelineStacks.GetKeys(Keys);
 		check(Keys.Num() > 0);
@@ -822,12 +802,12 @@ UInterchangeManager::ImportInternal(const FString& ContentPath, const UInterchan
 
 #if WITH_EDITORONLY_DATA
 		const bool bShowPipelineStacksConfigurationDialog = !bIsUnattended
-															&& InterchangeImportSettings.bShowPipelineStacksConfigurationDialog
+															&& FInterchangeProjectSettingsUtils::ShouldShowPipelineStacksConfigurationDialog(bImportScene, *SourceData)
 															&& !bImportAllWithDefault;
 #else
 		const bool bShowPipelineStacksConfigurationDialog = false;
 #endif
-															
+
 		const TMap<FName, FInterchangePipelineStack>& DefaultPipelineStacks = InterchangeImportSettings.PipelineStacks;
 
 		//If we reimport we want to load the original pipeline and the original pipeline settings
@@ -901,7 +881,7 @@ UInterchangeManager::ImportInternal(const FString& ContentPath, const UInterchan
 		}
 		else
 		{
-			FName PipelineStackName = InterchangeImportSettings.DefaultPipelineStack;
+			FName PipelineStackName = FInterchangeProjectSettingsUtils::GetDefaultPipelineStackName(ImportType == UE::Interchange::EImportType::ImportType_Scene, *SourceData);
 			if (RegisteredPipelineConfiguration && (bShowPipelineStacksConfigurationDialog || (!DefaultPipelineStacks.Contains(PipelineStackName) && !bIsUnattended)))
 			{
 				//Show the dialog, a plugin should have register this dialog. We use a plugin to be able to use editor code when doing UI
@@ -914,15 +894,15 @@ UInterchangeManager::ImportInternal(const FString& ContentPath, const UInterchan
 				{
 					bImportAllWithDefault = true;
 				}
-				PipelineStackName = InterchangeImportSettings.DefaultPipelineStack;
 			}
+
 			if (!bImportCancel)
 			{
 				if (!DefaultPipelineStacks.Contains(PipelineStackName))
 				{
-					if (DefaultPipelineStacks.Contains(InterchangeImportSettings.DefaultPipelineStack))
+					if (DefaultPipelineStacks.Contains(FInterchangeProjectSettingsUtils::GetDefaultPipelineStackName(ImportType == UE::Interchange::EImportType::ImportType_Scene, *SourceData)))
 					{
-						PipelineStackName = InterchangeImportSettings.DefaultPipelineStack;
+						PipelineStackName = FInterchangeProjectSettingsUtils::GetDefaultPipelineStackName(ImportType == UE::Interchange::EImportType::ImportType_Scene, *SourceData);
 					}
 					else
 					{
@@ -946,12 +926,17 @@ UInterchangeManager::ImportInternal(const FString& ContentPath, const UInterchan
 							if(UInterchangePipelineBase* GeneratedPipeline = UE::Interchange::GeneratePipelineInstance(Pipelines[GraphPipelineIndex]))
 							{
 								GeneratedPipeline->AdjustSettingsForContext(bImportScene ? EInterchangePipelineContext::SceneImport : EInterchangePipelineContext::AssetImport, nullptr);
-								//Load the settings for this pipeline
-								//The dialog is saving the settings for all default pipelines
-								GeneratedPipeline->LoadSettings(PipelineStackName);
+
+								if (FInterchangeProjectSettingsUtils::ShouldShowPipelineStacksConfigurationDialog(bImportScene, *SourceData))
+								{
+									//Load the settings for this pipeline
+									//The dialog is saving the settings for all default pipelines
+									GeneratedPipeline->LoadSettings(PipelineStackName);
+								}
+								
 								AsyncHelper->Pipelines.Add(GeneratedPipeline);
 
-								//We need to save the python pipeline asset because we cannot save an asset create with a python class
+								//We need to save the python pipeline asset because we cannot save an asset created with a python class
 								if (UInterchangePythonPipelineAsset* OriginalPipeline = Cast<UInterchangePythonPipelineAsset>(Pipelines[GraphPipelineIndex].TryLoad()))
 								{
 									if (UInterchangePythonPipelineAsset* DuplicatedPythonPipelineAsset = DuplicateObject<UInterchangePythonPipelineAsset>(OriginalPipeline, GetTransientPackage()))
diff --git a/Engine/Source/Runtime/Experimental/Interchange/Engine/Private/InterchangeProjectSettings.cpp b/Engine/Source/Runtime/Experimental/Interchange/Engine/Private/InterchangeProjectSettings.cpp
new file mode 100644
--- /dev/null
+++ b/Engine/Source/Runtime/Experimental/Interchange/Engine/Private/InterchangeProjectSettings.cpp
@@ -0,0 +1,97 @@
+// Copyright Epic Games, Inc. All Rights Reserved.
+#include "InterchangeProjectSettings.h"
+
+#include "InterchangeManager.h"
+#include "InterchangeSourceData.h"
+
+const FInterchangeImportSettings& FInterchangeProjectSettingsUtils::GetImportSettings(const UInterchangeProjectSettings& InterchangeProjectSettings, const bool bIsSceneImport)
+{
+	if (bIsSceneImport)
+	{
+		return InterchangeProjectSettings.SceneImportSettings;
+	}
+	else
+	{
+		return InterchangeProjectSettings.ContentImportSettings;
+	}
+}
+
+FInterchangeImportSettings& FInterchangeProjectSettingsUtils::GetMutableImportSettings(UInterchangeProjectSettings& InterchangeProjectSettings, const bool bIsSceneImport)
+{
+	if (bIsSceneImport)
+	{
+		return InterchangeProjectSettings.SceneImportSettings;
+	}
+	else
+	{
+		return InterchangeProjectSettings.ContentImportSettings;
+	}
+}
+
+const FInterchangeImportSettings& FInterchangeProjectSettingsUtils::GetDefaultImportSettings(const bool bIsSceneImport)
+{
+	return GetImportSettings(*GetDefault<UInterchangeProjectSettings>(), bIsSceneImport);
+}
+
+FInterchangeImportSettings& FInterchangeProjectSettingsUtils::GetMutableDefaultImportSettings(const bool bIsSceneImport)
+{
+	return GetMutableImportSettings(*GetMutableDefault<UInterchangeProjectSettings>(), bIsSceneImport);
+}
+
+FName FInterchangeProjectSettingsUtils::GetDefaultPipelineStackName(const bool bIsSceneImport, const UInterchangeSourceData& SourceData)
+{
+	const FInterchangeImportSettings& ImportSettings = GetDefaultImportSettings(bIsSceneImport);
+
+	FName DefaultPipelineStack = ImportSettings.DefaultPipelineStack;
+
+	if (!bIsSceneImport)
+	{
+		UE::Interchange::FScopedTranslator ScopedTranslator(&SourceData);
+
+		if (UInterchangeTranslatorBase* Translator = ScopedTranslator.GetTranslator())
+		{
+			EInterchangeTranslatorAssetType SupportedAssetTypes = Translator->GetSupportedAssetTypes();
+
+			const FInterchangeContentImportSettings& ContentImportSettings = GetDefault<UInterchangeProjectSettings>()->ContentImportSettings;
+			for (TMap<EInterchangeTranslatorAssetType, FName>::TConstIterator StackOverridesIt = ContentImportSettings.DefaultPipelineStackOverride.CreateConstIterator(); StackOverridesIt; ++StackOverridesIt)
+			{
+				if ((SupportedAssetTypes ^ StackOverridesIt->Key) < StackOverridesIt->Key)
+				{
+					DefaultPipelineStack = StackOverridesIt->Value;
+					break;
+				}
+			}
+		}
+	}
+
+	return DefaultPipelineStack;
+}
+
+bool FInterchangeProjectSettingsUtils::ShouldShowPipelineStacksConfigurationDialog(const bool bIsSceneImport, const UInterchangeSourceData& SourceData)
+{
+	const FInterchangeImportSettings& ImportSettings = GetDefaultImportSettings(bIsSceneImport);
+
+	bool bShowPipelineStacksConfigurationDialog = ImportSettings.bShowPipelineStacksConfigurationDialog;
+
+	if (!bIsSceneImport)
+	{
+		UE::Interchange::FScopedTranslator ScopedTranslator(&SourceData);
+
+		if (UInterchangeTranslatorBase* Translator = ScopedTranslator.GetTranslator())
+		{
+			EInterchangeTranslatorAssetType SupportedAssetTypes = Translator->GetSupportedAssetTypes();
+
+			const FInterchangeContentImportSettings& ContentImportSettings = GetDefault<UInterchangeProjectSettings>()->ContentImportSettings;
+			for (TMap<EInterchangeTranslatorAssetType, bool>::TConstIterator StackOverridesIt = ContentImportSettings.ShowPipelineStacksConfigurationDialogOverride.CreateConstIterator(); StackOverridesIt; ++StackOverridesIt)
+			{
+				if ((SupportedAssetTypes ^ StackOverridesIt->Key) < StackOverridesIt->Key)
+				{
+					bShowPipelineStacksConfigurationDialog = StackOverridesIt->Value;
+					break;
+				}
+			}
+		}
+	}
+
+	return bShowPipelineStacksConfigurationDialog;
+}
\ No newline at end of file
diff --git a/Engine/Source/Runtime/Experimental/Interchange/Engine/Public/InterchangeProjectSettings.h b/Engine/Source/Runtime/Experimental/Interchange/Engine/Public/InterchangeProjectSettings.h
--- a/Engine/Source/Runtime/Experimental/Interchange/Engine/Public/InterchangeProjectSettings.h
+++ b/Engine/Source/Runtime/Experimental/Interchange/Engine/Public/InterchangeProjectSettings.h
@@ -2,7 +2,6 @@
 
 #pragma once
 
-#include "CoreMinimal.h"
 #include "Engine/DeveloperSettings.h"
 #include "Engine/EngineTypes.h"
 #include "InterchangeFilePickerBase.h"
@@ -12,6 +11,8 @@
 
 #include "InterchangeProjectSettings.generated.h"
 
+class UInterchangeSourceData;
+
 USTRUCT()
 struct FInterchangePipelineStack
 {
@@ -43,6 +44,20 @@ struct FInterchangeImportSettings
 	bool bShowPipelineStacksConfigurationDialog = true;
 };
 
+USTRUCT()
+struct FInterchangeContentImportSettings : public FInterchangeImportSettings
+{
+	GENERATED_BODY()
+
+	/** This tell interchange which pipeline stack to select when importing.*/
+	UPROPERTY(EditAnywhere, Category = "Pipeline", Meta=(DisplayAfter="DefaultPipelineStack"))
+	TMap<EInterchangeTranslatorAssetType, FName> DefaultPipelineStackOverride;
+
+	/** This tell interchange which pipeline stack to select when importing.*/
+	UPROPERTY(EditAnywhere, Category = "Pipeline", Meta=(DisplayAfter="bShowPipelineStacksConfigurationDialog"))
+	TMap<EInterchangeTranslatorAssetType, bool> ShowPipelineStacksConfigurationDialogOverride;
+};
+
 UCLASS(config=Engine, meta=(DisplayName=Interchange), MinimalAPI)
 class UInterchangeProjectSettings : public UDeveloperSettings
 {
@@ -53,11 +68,11 @@ public:
 	 * Settings used when importing into the content browser.
 	 */
 	UPROPERTY(EditAnywhere, config, Category = "ImportContent")
-	FInterchangeImportSettings ContentImportSettings;
+	FInterchangeContentImportSettings ContentImportSettings;
 
 	/**
 	 * Settings used when importing into a level.
-	¸*/
+	 */
 	UPROPERTY(EditAnywhere, config, Category = "ImportIntoLevel")
 	FInterchangeImportSettings SceneImportSettings;
 
@@ -79,4 +94,17 @@ public:
 	 */
 	UPROPERTY(EditAnywhere, config, Category = "Editor Generic Pipeline Class")
 	TSoftClassPtr <UInterchangePipelineBase> GenericPipelineClass;
+};
+
+class INTERCHANGEENGINE_API FInterchangeProjectSettingsUtils
+{
+public:
+	static const FInterchangeImportSettings& GetImportSettings(const UInterchangeProjectSettings& InterchangeProjectSettings, const bool bIsSceneImport);
+	static FInterchangeImportSettings& GetMutableImportSettings(UInterchangeProjectSettings& InterchangeProjectSettings, const bool bIsSceneImport);
+	static const FInterchangeImportSettings& GetDefaultImportSettings(const bool bIsSceneImport);
+	static FInterchangeImportSettings& GetMutableDefaultImportSettings(const bool bIsSceneImport);
+
+	static FName GetDefaultPipelineStackName(const bool bIsSceneImport, const UInterchangeSourceData& SourceData);
+
+	static bool ShouldShowPipelineStacksConfigurationDialog(const bool bIsSceneImport, const UInterchangeSourceData& SourceData);
 };
\ No newline at end of file
