Interchange: Fix the reimport dialog to select the correct pipeline when doing a reimport.
#rnx

diff --git a/Engine/Plugins/Experimental/Interchange/Editor/Source/Pipelines/Private/SInterchangePipelineConfigurationDialog.cpp b/Engine/Plugins/Experimental/Interchange/Editor/Source/Pipelines/Private/SInterchangePipelineConfigurationDialog.cpp
--- a/Engine/Plugins/Experimental/Interchange/Editor/Source/Pipelines/Private/SInterchangePipelineConfigurationDialog.cpp
+++ b/Engine/Plugins/Experimental/Interchange/Editor/Source/Pipelines/Private/SInterchangePipelineConfigurationDialog.cpp
@@ -32,6 +32,8 @@
 
 extern INTERCHANGECORE_API bool GInterchangeEnableCustomPipelines;
 
+const FName ReimportPipelineName = TEXT("ReimportPipeline");
+
 /************************************************************************/
 /* FInterchangePipelineStacksTreeNodeItem Implementation                    */
 /************************************************************************/
@@ -62,7 +64,6 @@ void SInterchangePipelineStacksTreeView::Construct(const FArguments& InArgs)
 	bReimport = InArgs._bReimport && PipelineStack.Num() > 0;
 
 	//Build the FbxNodeInfoPtr tree data
-	const FName ReimportPipelineName = TEXT("ReimportPipeline");
 	const FName& DefaultPipelineStackName = bReimport ? ReimportPipelineName : FInterchangeProjectSettingsUtils::GetDefaultPipelineStackName(bSceneImport, *SourceData.Get());
 	const TMap<FName, FInterchangePipelineStack>& PipelineStacks = FInterchangeProjectSettingsUtils::GetDefaultImportSettings(bSceneImport).PipelineStacks;
 
@@ -127,7 +128,7 @@ void SInterchangePipelineStacksTreeView::Construct(const FArguments& InArgs)
 
 void SInterchangePipelineStacksTreeView::SelectDefaultItem()
 {
-	FName DefaultStackName = FInterchangeProjectSettingsUtils::GetDefaultPipelineStackName(bSceneImport, *SourceData.Get());
+	FName DefaultStackName = bReimport ? ReimportPipelineName : FInterchangeProjectSettingsUtils::GetDefaultPipelineStackName(bSceneImport, *SourceData.Get());
 
 	for (const TSharedPtr<FInterchangePipelineStacksTreeNodeItem>& PipelineStacksTreeNodeItem : RootNodeArray)
 	{
@@ -239,10 +240,12 @@ private:
 
 TSharedRef< ITableRow > SInterchangePipelineStacksTreeView::OnGenerateRowPipelineConfigurationTreeView(TSharedPtr<FInterchangePipelineStacksTreeNodeItem> Item, const TSharedRef< STableViewBase >& OwnerTable)
 {
+	FName DefaultStackName = bReimport ? ReimportPipelineName : FInterchangeProjectSettingsUtils::GetDefaultPipelineStackName(bSceneImport, *SourceData.Get());
+
 	TSharedRef< SInterchangePipelineStacksTreeViewItem > ReturnRow = SNew(SInterchangePipelineStacksTreeViewItem, OwnerTable)
 		.InterchangeNode(Item)
 		.bSceneImport(bSceneImport)
-		.DefaultPipelineStackName(FInterchangeProjectSettingsUtils::GetDefaultPipelineStackName(bSceneImport, *SourceData.Get()));
+		.DefaultPipelineStackName(DefaultStackName);
 	return ReturnRow;
 }
 void SInterchangePipelineStacksTreeView::OnGetChildrenPipelineConfigurationTreeView(TSharedPtr<FInterchangePipelineStacksTreeNodeItem> InParent, TArray< TSharedPtr<FInterchangePipelineStacksTreeNodeItem> >& OutChildren)
