Interchange
- Only show results browser when it contains warnings or errors.
- Never show the results browser during automated testing.

diff --git a/Engine/Plugins/Experimental/Interchange/Editor/Source/ResultsBrowser/Private/InterchangeResultsBrowserModule.cpp b/Engine/Plugins/Experimental/Interchange/Editor/Source/ResultsBrowser/Private/InterchangeResultsBrowserModule.cpp
--- a/Engine/Plugins/Experimental/Interchange/Editor/Source/ResultsBrowser/Private/InterchangeResultsBrowserModule.cpp
+++ b/Engine/Plugins/Experimental/Interchange/Editor/Source/ResultsBrowser/Private/InterchangeResultsBrowserModule.cpp
@@ -108,6 +108,25 @@ TSharedRef<SDockTab> FInterchangeResultsBrowserModule::SpawnInterchangeResults(c
 
 void FInterchangeResultsBrowserModule::OpenErrorBrowser(TStrongObjectPtr<UInterchangeResultsContainer> InResultsContainer)
 {
+	// Only showing when we have errors or warnings for now
+	const bool bShouldShow = [&InResultsContainer]()
+	{
+		for (UInterchangeResult* Result : InResultsContainer->GetResults())
+		{
+			if (Result->GetResultType() != EInterchangeResultType::Success)
+			{
+				return true;
+			}
+		}
+
+		return false;
+	}();
+
+	if (GIsAutomationTesting || !bShouldShow)
+	{
+		return;
+	}
+
 	ResultsContainer = InResultsContainer;
 	ResultsContainer->Finalize();
 	FGlobalTabmanager::Get()->TryInvokeTab(FName("InterchangeResults"));
