Interchange CIS fix for EngineTest Game_arm64 Hololens
#rnx

diff --git a/Engine/Plugins/Interchange/Editor/InterchangeEditor.uplugin b/Engine/Plugins/Interchange/Editor/InterchangeEditor.uplugin
--- a/Engine/Plugins/Interchange/Editor/InterchangeEditor.uplugin
+++ b/Engine/Plugins/Interchange/Editor/InterchangeEditor.uplugin
@@ -14,26 +14,47 @@
 	"CanContainContent": false,
 	"IsBetaVersion": false,
 	"Installed": false,
+	"SupportedTargetPlatforms": [
+		"Win64",
+		"Linux",
+		"Mac"
+	],
 	"Modules": [
 		{
 			"Name": "InterchangeEditor",
 			"Type": "Editor",
-			"LoadingPhase": "Default"
+			"LoadingPhase": "Default",
+			"TargetAllowList": [
+				"Editor",
+				"Program"
+			]
 		},
 		{
 			"Name": "InterchangeEditorPipelines",
 			"Type": "Editor",
-			"LoadingPhase": "Default"
+			"LoadingPhase": "Default",
+			"TargetAllowList": [
+				"Editor",
+				"Program"
+			]
 		},
 		{
 			"Name": "InterchangeResultsBrowser",
 			"Type": "Editor",
-			"LoadingPhase": "Default"
+			"LoadingPhase": "Default",
+			"TargetAllowList": [
+				"Editor",
+				"Program"
+			]
 		},
 		{
 			"Name": "InterchangeEditorUtilities",
 			"Type": "Editor",
-			"LoadingPhase": "Default"
+			"LoadingPhase": "Default",
+			"TargetAllowList": [
+				"Editor",
+				"Program"
+			]
 		}
 	],
 	"Plugins": [
diff --git a/Engine/Plugins/MeshPainting/MeshPainting.uplugin b/Engine/Plugins/MeshPainting/MeshPainting.uplugin
--- a/Engine/Plugins/MeshPainting/MeshPainting.uplugin
+++ b/Engine/Plugins/MeshPainting/MeshPainting.uplugin
@@ -33,7 +33,16 @@
 		},
 		{
 			"Name": "InterchangeEditor",
-			"Enabled": true
+			"Enabled": true,
+			"PlatformAllowList": [
+				"Win64",
+				"Linux",
+				"Mac"
+			],
+			"TargetAllowList": [
+				"Editor",
+				"Program"
+			]
 		}
 	],
 	"IsExperimentalVersion": false
