Interchange: Fix the open file dialog file filter list
#rnx

diff --git a/Engine/Plugins/Interchange/Editor/Source/Utilities/InterchangeEditorUtilities.Build.cs b/Engine/Plugins/Interchange/Editor/Source/Utilities/InterchangeEditorUtilities.Build.cs
--- a/Engine/Plugins/Interchange/Editor/Source/Utilities/InterchangeEditorUtilities.Build.cs
+++ b/Engine/Plugins/Interchange/Editor/Source/Utilities/InterchangeEditorUtilities.Build.cs
@@ -23,6 +23,7 @@ namespace UnrealBuildTool.Rules
 					"DesktopPlatform",
 					"Slate",
 					"SlateCore",
+					"UnrealEd"
 				}
 			);
 		}
diff --git a/Engine/Plugins/Interchange/Editor/Source/Utilities/Private/InterchangeOpenFileDialog.cpp b/Engine/Plugins/Interchange/Editor/Source/Utilities/Private/InterchangeOpenFileDialog.cpp
--- a/Engine/Plugins/Interchange/Editor/Source/Utilities/Private/InterchangeOpenFileDialog.cpp
+++ b/Engine/Plugins/Interchange/Editor/Source/Utilities/Private/InterchangeOpenFileDialog.cpp
@@ -3,65 +3,27 @@
 
 #include "Containers/UnrealString.h"
 #include "DesktopPlatformModule.h"
+#include "Framework/Application/SlateApplication.h"
 #include "IDesktopPlatform.h"
 #include "InterchangeFilePickerBase.h"
 #include "InterchangeManager.h"
 #include "InterchangeTranslatorBase.h"
+#include "ObjectTools.h"
 #include "UObject/Object.h"
 #include "UObject/ObjectMacros.h"
-#include "Framework/Application/SlateApplication.h"
 
 namespace UE::Interchange::Utilities::Private
 {
 	FString GetOpenFileDialogExtensions(const TArray<FString>& TranslatorFormats)
 	{
-		//Each format should be layout like this "fbx;Filmbox"
-		if (TranslatorFormats.Num() == 0)
-		{
-			return FString();
-		}
-		TArray<FString> Extensions;
-		TArray<FString> ExtensionDescriptions;
-		Extensions.Reserve(TranslatorFormats.Num());
-		ExtensionDescriptions.Reserve(TranslatorFormats.Num());
-		for (const FString& Format : TranslatorFormats)
-		{
-			int32 DelimiterIndex;
-			if (Format.FindChar(';', DelimiterIndex))
-			{
-				Extensions.Add(Format.Left(DelimiterIndex));
-				ExtensionDescriptions.Add(Format.RightChop(DelimiterIndex + 1));
-			}
-		}
-		const int32 ExtensionCount = Extensions.Num();
-		if (ExtensionCount == 0)
-		{
-			return FString();
-		}
+		FString FileTypes;
+		FString Extensions;
 
-		FString ExtensionStrAll = TEXT("All file extensions|");
-		FString ExtensionStr;
-		for (int32 ExtensionIndex = 0; ExtensionIndex < ExtensionCount; ++ExtensionIndex)
-		{
-			if (ExtensionIndex > 0)
-			{
-				ExtensionStrAll += TEXT(";*.");
-			}
-			else
-			{
-				ExtensionStrAll += TEXT("*.");
-			}
-			ExtensionStrAll += Extensions[ExtensionIndex];
-			ExtensionStr += ExtensionDescriptions[ExtensionIndex] + TEXT("|*.") + Extensions[ExtensionIndex];
-			if ((ExtensionIndex + 1) < ExtensionCount)
-			{
-				ExtensionStr += TEXT("|");
-			}
-		}
-		ExtensionStr += TEXT("Any files|*.*");
-		ExtensionStrAll += TEXT("|") + ExtensionStr;
+		ObjectTools::AppendFormatsFileExtensions(TranslatorFormats, FileTypes, Extensions);
+
+		const FString FormatString = FString::Printf(TEXT("All Files (%s)|%s|%s"), *Extensions, *Extensions, *FileTypes);
 
-		return ExtensionStrAll;
+		return FormatString;
 	}
 
 	bool FilePickerDialog(const FString& Extensions, FInterchangeFilePickerParameters& Parameters, TArray<FString>& OutFilenames)
