Interchange: Fixed the results browser stretching the result icon if the details are on multiple lines.

diff --git a/Engine/Plugins/Interchange/Editor/Source/ResultsBrowser/Private/SInterchangeResultsBrowserWindow.cpp b/Engine/Plugins/Interchange/Editor/Source/ResultsBrowser/Private/SInterchangeResultsBrowserWindow.cpp
--- a/Engine/Plugins/Interchange/Editor/Source/ResultsBrowser/Private/SInterchangeResultsBrowserWindow.cpp
+++ b/Engine/Plugins/Interchange/Editor/Source/ResultsBrowser/Private/SInterchangeResultsBrowserWindow.cpp
@@ -377,6 +377,7 @@ TSharedRef<SWidget> SInterchangeResultsBrowserListRow::GenerateWidgetForColumn(c
 				+ SHorizontalBox::Slot()
 				.Padding(RowPadding)
 				.AutoWidth()
+				.VAlign(VAlign_Center)
 				[
 					SNew(SImage)
 					.Image(FSlateIconFinder::FindIcon(GetImageFromResultType(Item->GetResultType())).GetOptionalIcon())
@@ -387,6 +388,7 @@ TSharedRef<SWidget> SInterchangeResultsBrowserListRow::GenerateWidgetForColumn(c
 			return SNew(SHorizontalBox)
 				+ SHorizontalBox::Slot()
 				.Padding(RowPadding)
+				.VAlign(VAlign_Center)
 				[
 					SNew(STextBlock)
 					.Text(FText::FromString(FPaths::GetCleanFilename(Item->SourceAssetName)))
@@ -398,6 +400,7 @@ TSharedRef<SWidget> SInterchangeResultsBrowserListRow::GenerateWidgetForColumn(c
 			return SNew(SHorizontalBox)
 				+ SHorizontalBox::Slot()
 				.Padding(RowPadding)
+				.VAlign(VAlign_Center)
 				[
 					SNew(STextBlock)
 					.Text(FText::FromString(FPackageName::ObjectPathToObjectName(Item->DestinationAssetName)))
@@ -409,6 +412,7 @@ TSharedRef<SWidget> SInterchangeResultsBrowserListRow::GenerateWidgetForColumn(c
 			return SNew(SHorizontalBox)
 				+ SHorizontalBox::Slot()
 				.Padding(RowPadding)
+				.VAlign(VAlign_Center)
 				[
 					SNew(STextBlock)
 					.Text(FText::FromString(Item->AssetType ? Item->AssetType->GetName() : FString()))
@@ -419,6 +423,7 @@ TSharedRef<SWidget> SInterchangeResultsBrowserListRow::GenerateWidgetForColumn(c
 			return SNew(SHorizontalBox)
 				+ SHorizontalBox::Slot()
 				.Padding(RowPadding)
+				.VAlign(VAlign_Center)
 				[
 					SNew(STextBlock)
 					.Text(Item->GetText())
