Interchange: refactor the import dialog to show the stacks and the pipelines.
#rnx

diff --git a/Engine/Plugins/Interchange/Editor/Source/Pipelines/Private/InterchangePipelineConfigurationGeneric.cpp b/Engine/Plugins/Interchange/Editor/Source/Pipelines/Private/InterchangePipelineConfigurationGeneric.cpp
--- a/Engine/Plugins/Interchange/Editor/Source/Pipelines/Private/InterchangePipelineConfigurationGeneric.cpp
+++ b/Engine/Plugins/Interchange/Editor/Source/Pipelines/Private/InterchangePipelineConfigurationGeneric.cpp
@@ -13,7 +13,9 @@
 
 #include UE_INLINE_GENERATED_CPP_BY_NAME(InterchangePipelineConfigurationGeneric)
 
-EInterchangePipelineConfigurationDialogResult UInterchangePipelineConfigurationGeneric::ShowPipelineConfigurationDialog(TWeakObjectPtr<UInterchangeSourceData> SourceData)
+EInterchangePipelineConfigurationDialogResult UInterchangePipelineConfigurationGeneric::ShowPipelineConfigurationDialog(TArray<FInterchangeStackInfo>& PipelineStacks
+	, TArray<UInterchangePipelineBase*>& OutPipelines
+	, TWeakObjectPtr<UInterchangeSourceData> SourceData)
 {
 	//Create and show the graph inspector UI dialog
 	TSharedPtr<SWindow> ParentWindow;
@@ -33,6 +35,8 @@ EInterchangePipelineConfigurationDialogResult UInterchangePipelineConfigurationG
 		.SourceData(SourceData)
 		.bSceneImport(false)
 		.bReimport(false)
+		.PipelineStacks(PipelineStacks)
+		.OutPipelines(&OutPipelines)
 	);
 
 	FSlateApplication::Get().AddModalWindow(Window, ParentWindow, false);
@@ -50,7 +54,9 @@ EInterchangePipelineConfigurationDialogResult UInterchangePipelineConfigurationG
 	return EInterchangePipelineConfigurationDialogResult::Import;
 }
 
-EInterchangePipelineConfigurationDialogResult UInterchangePipelineConfigurationGeneric::ShowScenePipelineConfigurationDialog(TWeakObjectPtr<UInterchangeSourceData> SourceData)
+EInterchangePipelineConfigurationDialogResult UInterchangePipelineConfigurationGeneric::ShowScenePipelineConfigurationDialog(TArray<FInterchangeStackInfo>& PipelineStacks
+	, TArray<UInterchangePipelineBase*>& OutPipelines
+	, TWeakObjectPtr<UInterchangeSourceData> SourceData)
 {
 	//Create and show the graph inspector UI dialog
 	TSharedPtr<SWindow> ParentWindow;
@@ -70,6 +76,8 @@ EInterchangePipelineConfigurationDialogResult UInterchangePipelineConfigurationG
 		.SourceData(SourceData)
 		.bSceneImport(true)
 		.bReimport(false)
+		.PipelineStacks(PipelineStacks)
+		.OutPipelines(&OutPipelines)
 	);
 
 	FSlateApplication::Get().AddModalWindow(Window, ParentWindow, false);
@@ -87,7 +95,9 @@ EInterchangePipelineConfigurationDialogResult UInterchangePipelineConfigurationG
 	return EInterchangePipelineConfigurationDialogResult::Import;
 }
 
-EInterchangePipelineConfigurationDialogResult UInterchangePipelineConfigurationGeneric::ShowReimportPipelineConfigurationDialog(TArray<UInterchangePipelineBase*>& PipelineStack, TWeakObjectPtr<UInterchangeSourceData> SourceData)
+EInterchangePipelineConfigurationDialogResult UInterchangePipelineConfigurationGeneric::ShowReimportPipelineConfigurationDialog(TArray<FInterchangeStackInfo>& PipelineStacks
+	, TArray<UInterchangePipelineBase*>& OutPipelines
+	, TWeakObjectPtr<UInterchangeSourceData> SourceData)
 {
 	//Create and show the graph inspector UI dialog
 	TSharedPtr<SWindow> ParentWindow;
@@ -107,7 +117,8 @@ EInterchangePipelineConfigurationDialogResult UInterchangePipelineConfigurationG
 		.SourceData(SourceData)
 		.bSceneImport(false)
 		.bReimport(true)
-		.PipelineStack(PipelineStack)
+		.PipelineStacks(PipelineStacks)
+		.OutPipelines(&OutPipelines)
 	);
 
 	FSlateApplication::Get().AddModalWindow(Window, ParentWindow, false);
diff --git a/Engine/Plugins/Interchange/Editor/Source/Pipelines/Private/SInterchangePipelineConfigurationDialog.cpp b/Engine/Plugins/Interchange/Editor/Source/Pipelines/Private/SInterchangePipelineConfigurationDialog.cpp
--- a/Engine/Plugins/Interchange/Editor/Source/Pipelines/Private/SInterchangePipelineConfigurationDialog.cpp
+++ b/Engine/Plugins/Interchange/Editor/Source/Pipelines/Private/SInterchangePipelineConfigurationDialog.cpp
@@ -9,6 +9,7 @@
 #include "InterchangeBlueprintPipelineBase.h"
 #include "InterchangeManager.h"
 #include "InterchangePipelineBase.h"
+#include "InterchangePipelineConfigurationBase.h"
 #include "InterchangeProjectSettings.h"
 #include "InterchangePythonPipelineBase.h"
 #include "InterchangeSourceData.h"
@@ -30,428 +31,192 @@
 
 #define LOCTEXT_NAMESPACE "InterchangePipelineConfiguration"
 
-extern INTERCHANGECORE_API bool GInterchangeEnableCustomPipelines;
+const FName ReimportStackName = TEXT("ReimportPipeline");
 
-const FName ReimportPipelineName = TEXT("ReimportPipeline");
-
-/************************************************************************/
-/* FInterchangePipelineStacksTreeNodeItem Implementation                    */
-/************************************************************************/
-
-void FInterchangePipelineStacksTreeNodeItem::AddReferencedObjects(FReferenceCollector& Collector)
+void SInterchangePipelineItem::Construct(
+	const FArguments& InArgs,
+	const TSharedRef<STableViewBase>& OwnerTable,
+	TObjectPtr<UInterchangePipelineBase> InPipelineElement)
 {
-	if (Pipeline)
-	{
-		Collector.AddReferencedObject(Pipeline);
-	}
+	PipelineElement = InPipelineElement;
+	TObjectPtr<UInterchangePipelineBase> PipelineElementPtr = PipelineElement.Get();
+	check(PipelineElementPtr.Get());
+	FText PipelineName = PipelineElementPtr.Get() ? FText::FromString(PipelineElementPtr->GetClass()->GetName()) : LOCTEXT("InvalidPipelineName", "Invalid Pipeline");
+	STableRow<TObjectPtr<UInterchangePipelineBase>>::Construct(
+		STableRow<TObjectPtr<UInterchangePipelineBase>>::FArguments()
+		.Content()
+		[
+			SNew(SBox)
+			.Padding(3.0f, 0.0f)
+			[
+				SNew(STextBlock)
+				.Text(PipelineName)
+			]
+		], OwnerTable);
 }
 
 /************************************************************************/
-/* SInterchangePipelineStacksTreeView Implementation                    */
+/* SInterchangePipelineConfigurationDialog Implementation                      */
 /************************************************************************/
 
-SInterchangePipelineStacksTreeView::~SInterchangePipelineStacksTreeView()
-{
-
-}
-
-void SInterchangePipelineStacksTreeView::Construct(const FArguments& InArgs)
+namespace UE::Private
 {
-	OnSelectionChangedDelegate = InArgs._OnSelectionChangedDelegate;
-	SourceData = InArgs._SourceData;
-	PipelineStack = InArgs._PipelineStack;
-	bSceneImport = InArgs._bSceneImport;
-	bReimport = InArgs._bReimport && PipelineStack.Num() > 0;
-
-	//Build the FbxNodeInfoPtr tree data
-	const FName& DefaultPipelineStackName = bReimport ? ReimportPipelineName : FInterchangeProjectSettingsUtils::GetDefaultPipelineStackName(bSceneImport, *SourceData.Get());
-	const TMap<FName, FInterchangePipelineStack>& PipelineStacks = FInterchangeProjectSettingsUtils::GetDefaultImportSettings(bSceneImport).PipelineStacks;
-
-	//In reimport we modify directly the asset pipeline
-	if (bReimport)
+	bool ContainStack(const TArray<FInterchangeStackInfo>& PipelineStacks, const FName StackName)
 	{
-		TSharedPtr<FInterchangePipelineStacksTreeNodeItem> StackNode = MakeShared<FInterchangePipelineStacksTreeNodeItem>();
-		StackNode->StackName = ReimportPipelineName;
-		StackNode->Pipeline = nullptr;
-		for (int32 PipelineIndex = 0; PipelineIndex < PipelineStack.Num(); ++PipelineIndex)
-		{
-			if(UInterchangePipelineBase* GeneratedPipeline = PipelineStack[PipelineIndex])
+		return PipelineStacks.FindByPredicate([StackName](const FInterchangeStackInfo& StackInfo)
 			{
-				//When we re-import we save the setting so ResetToDefault can simply just reload the settings
-				GeneratedPipeline->SaveSettings(ReimportPipelineName);
-				TSharedPtr<FInterchangePipelineStacksTreeNodeItem> PipelineNode = MakeShared<FInterchangePipelineStacksTreeNodeItem>();
-				PipelineNode->StackName = ReimportPipelineName;
-				//Currently the re-import never lock properties
-				PipelineNode->Pipeline = GeneratedPipeline;
-				StackNode->Childrens.Add(PipelineNode);
-			}
-		}
-		RootNodeArray.Add(StackNode);
-	}
-	else
-	{
-		for (const TPair<FName, FInterchangePipelineStack>& NameAndPipelineStack : PipelineStacks)
-		{
-			TSharedPtr<FInterchangePipelineStacksTreeNodeItem> StackNode = MakeShared<FInterchangePipelineStacksTreeNodeItem>();
-			StackNode->StackName = NameAndPipelineStack.Key;
-			StackNode->Pipeline = nullptr;
-			const FInterchangePipelineStack& InterchangePipelineStack = NameAndPipelineStack.Value;
-			for (int32 PipelineIndex = 0; PipelineIndex < InterchangePipelineStack.Pipelines.Num(); ++PipelineIndex)
-			{
-				if (UInterchangePipelineBase* GeneratedPipeline = UE::Interchange::GeneratePipelineInstance(InterchangePipelineStack.Pipelines[PipelineIndex]))
-				{
-					GeneratedPipeline->AdjustSettingsForContext((bSceneImport ? EInterchangePipelineContext::SceneImport : EInterchangePipelineContext::AssetImport), nullptr);
-					//Load the settings for this pipeline
-					GeneratedPipeline->LoadSettings(NameAndPipelineStack.Key);
-					GeneratedPipeline->PreDialogCleanup(NameAndPipelineStack.Key);
-					TSharedPtr<FInterchangePipelineStacksTreeNodeItem> PipelineNode = MakeShared<FInterchangePipelineStacksTreeNodeItem>();
-					PipelineNode->StackName = NameAndPipelineStack.Key;
-					PipelineNode->Pipeline = GeneratedPipeline;
-					StackNode->Childrens.Add(PipelineNode);
-				}
-			}
-			RootNodeArray.Add(StackNode);
-		}
-	}
-
-	STreeView::Construct
-	(
-		STreeView::FArguments()
-		.TreeItemsSource(&RootNodeArray)
-		.SelectionMode(ESelectionMode::Single)
-		.OnGenerateRow(this, &SInterchangePipelineStacksTreeView::OnGenerateRowPipelineConfigurationTreeView)
-		.OnGetChildren(this, &SInterchangePipelineStacksTreeView::OnGetChildrenPipelineConfigurationTreeView)
-		.OnContextMenuOpening(this, &SInterchangePipelineStacksTreeView::OnOpenContextMenu)
-		.OnSelectionChanged(this, &SInterchangePipelineStacksTreeView::OnTreeViewSelectionChanged)
-	);
-}
-
-void SInterchangePipelineStacksTreeView::SelectDefaultItem()
-{
-	FName DefaultStackName = bReimport ? ReimportPipelineName : FInterchangeProjectSettingsUtils::GetDefaultPipelineStackName(bSceneImport, *SourceData.Get());
-
-	for (const TSharedPtr<FInterchangePipelineStacksTreeNodeItem>& PipelineStacksTreeNodeItem : RootNodeArray)
-	{
-		if (PipelineStacksTreeNodeItem->StackName == DefaultStackName && !PipelineStacksTreeNodeItem->Childrens.IsEmpty())
-		{
-			const bool bShouldExpandItem = true;
-			SetItemExpansion(PipelineStacksTreeNodeItem, bShouldExpandItem);
-			SetSelection(PipelineStacksTreeNodeItem->Childrens[0]);
-
-			break;
-		}
+				return (StackInfo.StackName == StackName);
+			}) != nullptr;
 	}
 }
 
-/** The item used for visualizing the class in the tree. */
-class SInterchangePipelineStacksTreeViewItem : public STableRow< TSharedPtr<FInterchangePipelineStacksTreeNodeItem> >
-{
-public:
-
-	SLATE_BEGIN_ARGS(SInterchangePipelineStacksTreeViewItem)
-		: _InterchangeNode(nullptr)
-	{}
-
-		/** The item content. */
-		SLATE_ARGUMENT(TSharedPtr<FInterchangePipelineStacksTreeNodeItem>, InterchangeNode)
-		SLATE_ARGUMENT(bool, bSceneImport)
-		SLATE_ARGUMENT(FName, DefaultPipelineStackName)
-	SLATE_END_ARGS()
-
-	/**
-	* Construct the widget
-	*
-	* @param InArgs   A declaration from which to construct the widget
-	*/
-	void Construct(const FArguments& InArgs, const TSharedRef<STableViewBase>& InOwnerTableView)
-	{
-		InterchangeNode = InArgs._InterchangeNode;
-		bSceneImport = InArgs._bSceneImport;
-		DefaultPipelineStackName = InArgs._DefaultPipelineStackName;
-
-		//This is suppose to always be valid
-		check(InterchangeNode);
-		const bool bIsPipelineStackNode = InterchangeNode->Pipeline == nullptr;
-
-		//Prepare the tooltip
-		FString Tooltip;
-		//FString Tooltip = bIsPipelineStackNode ? InterchangeNode->GetDisplayLabel();
-		FText NodeDisplayLabel = FText::FromName(bIsPipelineStackNode ? InterchangeNode->StackName : InterchangeNode->Pipeline->GetClass()->GetFName());
-		this->ChildSlot
-		[
-			SNew(SHorizontalBox)
-			+ SHorizontalBox::Slot()
-			.AutoWidth()
-			[
-				SNew(SExpanderArrow, SharedThis(this))
-			]
-			+ SHorizontalBox::Slot()
-			.AutoWidth()
-			.Padding(0.0f, 2.0f, 6.0f, 2.0f)
-			[
-				SNew(SImage)
-				.Image(this, &SInterchangePipelineStacksTreeViewItem::GetImageItemIcon)
-			]
-			+ SHorizontalBox::Slot()
-			.FillWidth(1.0f)
-			.Padding(0.0f, 3.0f, 6.0f, 3.0f)
-			.VAlign(VAlign_Center)
-			[
-				SNew(STextBlock)
-				.Text(NodeDisplayLabel)
-				.ToolTipText(FText::FromString(Tooltip))
-			]
-		];
-
-		STableRow< TSharedPtr<FInterchangePipelineStacksTreeNodeItem> >::ConstructInternal(
-			STableRow::FArguments()
-			.ShowSelection(true),
-			InOwnerTableView
-		);
-	}
-
-private:
-	const FSlateBrush* GetImageItemIcon() const
-	{
-		//This is suppose to always be valid
-		check(InterchangeNode);
-		const bool bIsPipelineStackNode = InterchangeNode->Pipeline == nullptr;
-		const bool bIsDefaultStackNode = bIsPipelineStackNode && (DefaultPipelineStackName == InterchangeNode->StackName);
-		const FSlateBrush * TypeIcon = nullptr;
-		FName IconName = bIsDefaultStackNode ? "PipelineConfigurationIcon.PipelineStackDefault" : bIsPipelineStackNode ? "PipelineConfigurationIcon.PipelineStack" : "PipelineConfigurationIcon.Pipeline";
-		if (IconName != NAME_None)
-		{
-			const FSlateIcon SlateIcon = FSlateIconFinder::FindIcon(IconName);
-			TypeIcon = SlateIcon.GetOptionalIcon();
-		}
-
-		if (!TypeIcon)
-		{
-			TypeIcon = FSlateIconFinder::FindIconBrushForClass(AActor::StaticClass());
-		}
-		return TypeIcon;
-	}
-
-	/** The node to build the tree view row from. */
-	TSharedPtr<FInterchangePipelineStacksTreeNodeItem> InterchangeNode = nullptr;
-	bool bSceneImport = false;
-	FName DefaultPipelineStackName;
-};
-
-TSharedRef< ITableRow > SInterchangePipelineStacksTreeView::OnGenerateRowPipelineConfigurationTreeView(TSharedPtr<FInterchangePipelineStacksTreeNodeItem> Item, const TSharedRef< STableViewBase >& OwnerTable)
-{
-	FName DefaultStackName = bReimport ? ReimportPipelineName : FInterchangeProjectSettingsUtils::GetDefaultPipelineStackName(bSceneImport, *SourceData.Get());
-
-	TSharedRef< SInterchangePipelineStacksTreeViewItem > ReturnRow = SNew(SInterchangePipelineStacksTreeViewItem, OwnerTable)
-		.InterchangeNode(Item)
-		.bSceneImport(bSceneImport)
-		.DefaultPipelineStackName(DefaultStackName);
-	return ReturnRow;
-}
-void SInterchangePipelineStacksTreeView::OnGetChildrenPipelineConfigurationTreeView(TSharedPtr<FInterchangePipelineStacksTreeNodeItem> InParent, TArray< TSharedPtr<FInterchangePipelineStacksTreeNodeItem> >& OutChildren)
+SInterchangePipelineConfigurationDialog::SInterchangePipelineConfigurationDialog()
 {
-	for (int32 ChildIndex = 0; ChildIndex < InParent->Childrens.Num(); ++ChildIndex)
-	{
-		TSharedPtr<FInterchangePipelineStacksTreeNodeItem> ChildNode = InParent->Childrens[ChildIndex];
-		if (!ChildNode.IsValid())
-			continue;
-		OutChildren.Add(ChildNode);
-	}
+	PipelineConfigurationDetailsView = nullptr;
+	OwnerWindow = nullptr;
 }
 
-void SInterchangePipelineStacksTreeView::RecursiveSetExpand(TSharedPtr<FInterchangePipelineStacksTreeNodeItem> Node, bool ExpandState)
+SInterchangePipelineConfigurationDialog::~SInterchangePipelineConfigurationDialog()
 {
-	SetItemExpansion(Node, ExpandState);
-	for (int32 ChildIndex = 0; ChildIndex < Node->Childrens.Num(); ++ChildIndex)
+	if (TSharedPtr<SWindow> OwnerWindowPinned = OwnerWindow.Pin())
 	{
-		TSharedPtr<FInterchangePipelineStacksTreeNodeItem> ChildNode = Node->Childrens[ChildIndex];
-		if (!ChildNode.IsValid())
-			continue;
-		RecursiveSetExpand(ChildNode, ExpandState);
+		OwnerWindowPinned->GetOnWindowClosedEvent().RemoveAll(this);
 	}
 }
 
-FReply SInterchangePipelineStacksTreeView::OnExpandAll()
+TSharedRef<SBox> SInterchangePipelineConfigurationDialog::SpawnPipelineConfiguration()
 {
-	for (TSharedPtr<FInterchangePipelineStacksTreeNodeItem> Node : RootNodeArray)
+	AvailableStacks.Reset();
+	TSharedPtr<FString> SelectedStack;
+	if (bReimport)
 	{
-		if (!ensure(Node))
-		{
-			continue;
-		}
-		RecursiveSetExpand(Node, true);
+		CurrentStackName = ReimportStackName;
 	}
-	return FReply::Handled();
-}
-
-FReply SInterchangePipelineStacksTreeView::OnCollapseAll()
-{
-	for (TSharedPtr<FInterchangePipelineStacksTreeNodeItem> Node : RootNodeArray)
+	else
 	{
-		if (!ensure(Node))
-		{
-			continue;
-		}
-		RecursiveSetExpand(Node, false);
+		CurrentStackName = FInterchangeProjectSettingsUtils::GetDefaultPipelineStackName(bSceneImport, *SourceData);
 	}
-	return FReply::Handled();
-}
 
-TSharedPtr<SWidget> SInterchangePipelineStacksTreeView::OnOpenContextMenu()
-{
-	if (bReimport)
+	//In case we do not have a valid stack name use the first stack
+	FName FirstStackName = PipelineStacks.Num() > 0 ? PipelineStacks[0].StackName : CurrentStackName;
+	if (!UE::Private::ContainStack(PipelineStacks, CurrentStackName))
 	{
-		return SNullWidget::NullWidget;
+		CurrentStackName = FirstStackName;
 	}
 
-	// Build up the menu for a selection
-	const bool bCloseAfterSelection = true;
-	FMenuBuilder MenuBuilder(bCloseAfterSelection, TSharedPtr<FUICommandList>());
-
-	TArray<TSharedPtr<FInterchangePipelineStacksTreeNodeItem>> SelectedNodes;
-	const auto NumSelectedItems = GetSelectedItems(SelectedNodes);
-
-	if (SelectedNodes.Num() == 1)
+	for(FInterchangeStackInfo& Stack : PipelineStacks)
 	{
-		TSharedPtr<FInterchangePipelineStacksTreeNodeItem> SelectNode = SelectedNodes[0];
-		if (SelectNode->Pipeline == nullptr)
+		TSharedPtr<FString> StackNamePtr = MakeShared<FString>(Stack.StackName.ToString());
+		if (CurrentStackName == Stack.StackName)
 		{
-			// We always create a section here, even if there is no parent so that clients can still extend the menu
-			MenuBuilder.BeginSection("TreeViewContextMenuStackNodeSection");
+			for (UInterchangePipelineBase* DefaultPipeline : Stack.Pipelines)
 			{
-				const FSlateIcon DefaultIcon(FAppStyle::GetAppStyleSetName(), "Icons.Default");
-				MenuBuilder.AddMenuEntry(LOCTEXT("SetAsDefaultMenuAction", "Set As Default Stack"), FText(), DefaultIcon, FUIAction(FExecuteAction::CreateSP(this, &SInterchangePipelineStacksTreeView::SetAsDefaultStack, SelectNode->StackName)));
+				check(DefaultPipeline);
+				if (UInterchangePipelineBase* GeneratedPipeline = UE::Interchange::GeneratePipelineInstance(DefaultPipeline))
+				{
+					GeneratedPipeline->TransferAdjustSettings(DefaultPipeline);
+					if (Stack.StackName == ReimportStackName)
+					{
+						
+						//We save the pipeline settings to allow Reset to Default to work
+						GeneratedPipeline->SaveSettings(Stack.StackName);
+					}
+					else
+					{
+						//Load the settings for this pipeline
+						GeneratedPipeline->LoadSettings(Stack.StackName);
+						GeneratedPipeline->PreDialogCleanup(Stack.StackName);
+					}
+					PipelineListViewItems.Add(GeneratedPipeline);
+				}
 			}
-			MenuBuilder.EndSection();
+			SelectedStack = StackNamePtr;
 		}
+		AvailableStacks.Add(StackNamePtr);
 	}
-	
 
-	return MenuBuilder.MakeWidget();
-}
 
-void SInterchangePipelineStacksTreeView::SetAsDefaultStack(FName NewDefaultStackValue)
-{
-	UInterchangeProjectSettings* InterchangeProjectSettingsCDO = GetMutableDefault<UInterchangeProjectSettings>();
-	FName& DefaultPipelineStackName = [this, InterchangeProjectSettingsCDO]() -> FName&
-	{
-		if (bSceneImport)
-		{
-			return InterchangeProjectSettingsCDO->SceneImportSettings.DefaultPipelineStack;
-		}
-		else
+	PipelinesListView = SNew(SPipelineListViewType)
+		.SelectionMode(ESelectionMode::Single)
+		.ListItemsSource(&PipelineListViewItems)
+		.OnGenerateRow(this, &SInterchangePipelineConfigurationDialog::MakePipelineListRowWidget)
+		.OnSelectionChanged(this, &SInterchangePipelineConfigurationDialog::OnPipelineSelectionChanged);
+
+	TSharedPtr<STextComboBox> TextComboBoxPtr;
+	//Only use a combo box if there is more then one stack
+	if (AvailableStacks.Num() > 1)
+	{
+		TextComboBoxPtr = SNew(STextComboBox)
+			.OptionsSource(&AvailableStacks)
+			.OnSelectionChanged(this, &SInterchangePipelineConfigurationDialog::OnStackSelectionChanged);
+		if (SelectedStack.IsValid())
 		{
-			return InterchangeProjectSettingsCDO->ContentImportSettings.DefaultPipelineStack;
+			TextComboBoxPtr->SetSelectedItem(SelectedStack);
 		}
-	}();
-
-	const TMap<FName, FInterchangePipelineStack>& PipelineStacks =
-		[this, InterchangeProjectSettingsCDO]() -> const TMap<FName, FInterchangePipelineStack>&
-		{
-			if (bSceneImport)
-			{
-				return InterchangeProjectSettingsCDO->SceneImportSettings.PipelineStacks;
-			}
-			else
-			{
-				return InterchangeProjectSettingsCDO->ContentImportSettings.PipelineStacks;
-			}
-		}();
-
-	if (PipelineStacks.Contains(NewDefaultStackValue))
-	{
-		DefaultPipelineStackName = NewDefaultStackValue;
-		InterchangeProjectSettingsCDO->SaveConfig(); //This ensure the default pipeline stack name is save into the local config
 	}
-}
 
-void SInterchangePipelineStacksTreeView::OnTreeViewSelectionChanged(TSharedPtr<FInterchangePipelineStacksTreeNodeItem> Item, ESelectInfo::Type SelectionType)
-{
-	if (SelectionMode.Get() == ESelectionMode::None)
-	{
-		return;
-	}
+	FText CurrentStackText = LOCTEXT("CurrentStackText", "Stack: ");
 
-	if (OnSelectionChangedDelegate.IsBound())
+	TSharedPtr<SWidget> StackTextComboBox;
+	if (!TextComboBoxPtr.IsValid())
 	{
-		OnSelectionChangedDelegate.ExecuteIfBound(Item, SelectionType);
+		CurrentStackText = FText::Format(LOCTEXT("CurrentStackTextNoComboBox", "Stack: {0}"), FText::FromName(CurrentStackName));
+		StackTextComboBox = SNew(SBox)
+		[
+			SNew(STextBlock)
+			.Text(CurrentStackText)
+		];
 	}
-}
-
-
-/************************************************************************/
-/* SInterchangePipelineConfigurationDialog Implementation                      */
-/************************************************************************/
-
-SInterchangePipelineConfigurationDialog::SInterchangePipelineConfigurationDialog()
-{
-	PipelineConfigurationTreeView = nullptr;
-	PipelineConfigurationDetailsView = nullptr;
-	OwnerWindow = nullptr;
-}
-
-SInterchangePipelineConfigurationDialog::~SInterchangePipelineConfigurationDialog()
-{
-	if (TSharedPtr<SWindow> OwnerWindowPinned = OwnerWindow.Pin())
+	else
 	{
-		OwnerWindowPinned->GetOnWindowClosedEvent().RemoveAll(this);
+		StackTextComboBox = SNew(SHorizontalBox)
+		+ SHorizontalBox::Slot()
+		.VAlign(VAlign_Center)
+		.AutoWidth()
+		[
+			SNew(SBox)
+			[
+				SNew(STextBlock)
+				.Text(CurrentStackText)
+			]
+		]
+		+ SHorizontalBox::Slot()
+		.VAlign(VAlign_Center)
+		.AutoWidth()
+		[
+			TextComboBoxPtr.ToSharedRef()
+		];
 	}
-}
-
-TSharedRef<SBox> SInterchangePipelineConfigurationDialog::SpawnPipelineConfiguration()
-{
-	//Create the treeview
-	PipelineConfigurationTreeView = SNew(SInterchangePipelineStacksTreeView)
-		.OnSelectionChangedDelegate(this, &SInterchangePipelineConfigurationDialog::OnSelectionChanged)
-		.SourceData(SourceData)
-		.bSceneImport(bSceneImport)
-		.bReimport(bReimport)
-		.PipelineStack(PipelineStack);
-
-	const float PipelineStackTreeViewWidth = GInterchangeEnableCustomPipelines ? 0.4f : 0.f;
 
 	TSharedPtr<SBox> InspectorBox;
 	TSharedRef<SBox> PipelineConfigurationPanelBox = SNew(SBox)
 	[
-		SNew(SSplitter)
-		.Orientation(Orient_Horizontal)
-		+ SSplitter::Slot()
-		.Value(PipelineStackTreeViewWidth)
-		.Resizable(GInterchangeEnableCustomPipelines)
+		SNew(SVerticalBox)
+		+ SVerticalBox::Slot()
+		.AutoHeight()
 		[
 			SNew(SVerticalBox)
 			+ SVerticalBox::Slot()
 			.HAlign(HAlign_Left)
 			.AutoHeight()
 			[
-				SNew(SUniformGridPanel)
-				.SlotPadding(2)
-				+ SUniformGridPanel::Slot(0, 0)
-				[
-					SNew(SButton)
-					.HAlign(HAlign_Center)
-					.Text(LOCTEXT("FbxOptionWindow_Scene_ExpandAll", "Expand All"))
-					.OnClicked(PipelineConfigurationTreeView.Get(), &SInterchangePipelineStacksTreeView::OnExpandAll)
-				]
-				+ SUniformGridPanel::Slot(1, 0)
-				[
-					SNew(SButton)
-					.HAlign(HAlign_Center)
-					.Text(LOCTEXT("FbxOptionWindow_Scene_CollapseAll", "Collapse All"))
-					.OnClicked(PipelineConfigurationTreeView.Get(), &SInterchangePipelineStacksTreeView::OnCollapseAll)
-				]
+				StackTextComboBox.ToSharedRef()
 			]
 			+ SVerticalBox::Slot()
-			.FillHeight(1.0f)
+			.Padding(0.0f, 5.0f)
+			.AutoHeight()
 			[
 				SNew(SBox)
+				.MinDesiredHeight(35)
+				.MaxDesiredHeight(110)
 				[
-					PipelineConfigurationTreeView.ToSharedRef()
+					PipelinesListView.ToSharedRef()
 				]
 			]
 		]
-		+ SSplitter::Slot()
-		.Value(1.f - PipelineStackTreeViewWidth)
+		+ SVerticalBox::Slot()
+		.Padding(0.0f, 5.0f, 0.0f, 2.0f)
+		.FillHeight(1.0f)
 		[
 			SAssignNew(InspectorBox, SBox)
 		]
@@ -482,7 +247,10 @@ void SInterchangePipelineConfigurationDialog::Construct(const FArguments& InArgs
 	SourceData = InArgs._SourceData;
 	bSceneImport = InArgs._bSceneImport;
 	bReimport = InArgs._bReimport;
-	PipelineStack = InArgs._PipelineStack;
+	PipelineStacks = InArgs._PipelineStacks;
+	OutPipelines = InArgs._OutPipelines;
+
+	check(OutPipelines);
 
 	check(OwnerWindow.IsValid());
 	if (TSharedPtr<SWindow> OwnerWindowPinned = OwnerWindow.Pin())
@@ -500,12 +268,12 @@ void SInterchangePipelineConfigurationDialog::Construct(const FArguments& InArgs
 			+ SVerticalBox::Slot()
 			.AutoHeight()
 			.HAlign(HAlign_Fill)
-			.Padding(2)
+			.Padding(0.0f, 2.0f)
 			[
 				SNew(SHorizontalBox)
 				+SHorizontalBox::Slot()
 				.FillWidth(1.0f)
-				.Padding(2)
+				.Padding(0.0f, 2.0f)
 				[
 					SNew(STextBlock)
 					.Text(this, &SInterchangePipelineConfigurationDialog::GetSourceDescription)
@@ -513,24 +281,38 @@ void SInterchangePipelineConfigurationDialog::Construct(const FArguments& InArgs
 				+SHorizontalBox::Slot()
 				.AutoWidth()
 				.HAlign(HAlign_Right)
-				.Padding(10, 2, 2, 2)
+				.Padding(10.0f, 2.0f, 0.0f, 2.0f)
 				[
 					SNew(SButton)
 					.HAlign(HAlign_Center)
-					.Text(LOCTEXT("InspectorGraphWindow_ResetDefault", "Reset to Default"))
+					.IsEnabled_Lambda([this]()
+						{
+							return PipelinesListView->GetNumItemsSelected() == 1;
+						})
+					.Text(LOCTEXT("SInterchangePipelineConfigurationDialog_ResetToPipelineAsset", "Reset Selected Pipeline"))
+					.ToolTipText_Lambda([bReimportClosure = bReimport]()
+						{
+							if (bReimportClosure)
+							{
+								return LOCTEXT("SInterchangePipelineConfigurationDialog_ResetToPipelineAsset_TooltipReimport", "Reset the selected pipeline properties to the asset import data pipeline properties.");
+							}
+							else
+							{
+								return LOCTEXT("SInterchangePipelineConfigurationDialog_ResetToPipelineAsset_Tooltip", "Reset the selected pipeline properties to the stack pipeline properties.");
+							}
+						})
 					.OnClicked(this, &SInterchangePipelineConfigurationDialog::OnResetToDefault)
 				]
 			]
 			+ SVerticalBox::Slot()
 			.FillHeight(1.0f)
-			.Padding(2)
 			[
 				SpawnPipelineConfiguration()
 			]
 			+ SVerticalBox::Slot()
 			.AutoHeight()
 			.HAlign(HAlign_Right)
-			.Padding(2)
+			.Padding(0.0f, 2.0f)
 			[
 				SNew(SHorizontalBox)
 				+ SHorizontalBox::Slot()
@@ -585,20 +367,11 @@ void SInterchangePipelineConfigurationDialog::Construct(const FArguments& InArgs
 		]
 	];
 
-	PipelineConfigurationTreeView->SelectDefaultItem();
-}
-
-void SInterchangePipelineConfigurationDialog::OnSelectionChanged(TSharedPtr<FInterchangePipelineStacksTreeNodeItem> Item, ESelectInfo::Type SelectionType)
-{
-	CurrentSelectedPipeline = nullptr;
-	FName StackName = NAME_None;
-	if (Item.IsValid())
+	//Select the first pipeline
+	if (PipelineListViewItems.Num() > 0)
 	{
-		CurrentSelectedPipeline = Item->Pipeline;
-		StackName = Item->StackName;
+		PipelinesListView->SetSelection(PipelineListViewItems[0], ESelectInfo::Direct);
 	}
-	PipelineConfigurationDetailsView->SetObject(CurrentSelectedPipeline.Get());
-	CurrentStackName = StackName;
 }
 
 bool SInterchangePipelineConfigurationDialog::IsPropertyVisible(const FPropertyAndParent& PropertyAndParent) const
@@ -629,16 +402,6 @@ FText SInterchangePipelineConfigurationDialog::GetSourceDescription() const
 	return ActionDescription;
 }
 
-void SInterchangePipelineConfigurationDialog::RecursiveIterateNode(TSharedPtr<FInterchangePipelineStacksTreeNodeItem>& ParentNode, TFunctionRef<void(TSharedPtr<FInterchangePipelineStacksTreeNodeItem>&)> IterationLambda)
-{
-	IterationLambda(ParentNode);
-
-	for (int32 ChildIndex = 0; ChildIndex < ParentNode->Childrens.Num(); ++ChildIndex)
-	{
-		RecursiveIterateNode(ParentNode->Childrens[ChildIndex], IterationLambda);
-	}
-}
-
 FReply SInterchangePipelineConfigurationDialog::OnResetToDefault()
 {
 	FReply Result = FReply::Handled();
@@ -653,59 +416,36 @@ FReply SInterchangePipelineConfigurationDialog::OnResetToDefault()
 	{
 		if (UInterchangePipelineBase* Pipeline = Cast<UInterchangePipelineBase>(WeakObject.Get()))
 		{
-			if (bReimport)
+			const UClass* PipelineClass = Pipeline->GetClass();
+
+			for (FInterchangeStackInfo& Stack : PipelineStacks)
 			{
-				//When we re-import we save the settings before showing the dialog, so the reset to default simply reload
-				//the original settings when the dialog was loaded.
-				for (int32 PipelineIndex = 0; PipelineIndex < PipelineStack.Num(); ++PipelineIndex)
+				if (Stack.StackName != CurrentStackName)
 				{
-					if (Pipeline == PipelineStack[PipelineIndex])
-					{
-						Pipeline->LoadSettings(CurrentStackName);
-						PipelineConfigurationDetailsView->SetObject(Pipeline, true);
-						break;
-					}
+					continue;
 				}
-			}
-			else
-			{
-				//The reset to default when doing a import is to duplicate again the pipeline stack reference asset.
-				const TMap<FName, FInterchangePipelineStack>& DefaultPipelineStacks = FInterchangeProjectSettingsUtils::GetDefaultImportSettings(bSceneImport).PipelineStacks;
-				//Find the stack node, so we can update the pipeline
-				TSharedPtr<FInterchangePipelineStacksTreeNodeItem> StackNode;
-				TArray<TSharedPtr<FInterchangePipelineStacksTreeNodeItem>>& RootNodeArray = PipelineConfigurationTreeView->GetMutableRootNodeArray();
-				for (TSharedPtr<FInterchangePipelineStacksTreeNodeItem>& RootNode : RootNodeArray)
+				for (UInterchangePipelineBase* DefaultPipeline : Stack.Pipelines)
 				{
-					RecursiveIterateNode(RootNode, [this, &Pipeline, &DefaultPipelineStacks](TSharedPtr<FInterchangePipelineStacksTreeNodeItem>& NodeItem)
+					//We assume the pipelines inside one stack are all different classes, we use the class to know which default asset we need to duplicate
+					if (DefaultPipeline->GetClass() == PipelineClass)
 					{
-						if (NodeItem->StackName == CurrentStackName && NodeItem->Pipeline == Pipeline)
+						for(int32 PipelineIndex = 0; PipelineIndex < PipelineListViewItems.Num(); ++PipelineIndex)
 						{
-							//We assume the pipelines inside one stack are all different classes, we use the class to know which default asset we need to duplicate
-							const UClass* PipelineClass = NodeItem->Pipeline->GetClass();
-							for (const TPair<FName, FInterchangePipelineStack>& NameAndPipelineStack : DefaultPipelineStacks)
+							TObjectPtr<UInterchangePipelineBase> PipelineElement = PipelineListViewItems[PipelineIndex];
+							if (PipelineElement.Get() == Pipeline)
 							{
-								if(CurrentStackName != NameAndPipelineStack.Key)
-								{
-									continue;
-								}
-								const FInterchangePipelineStack& InterchangePipelineStack = NameAndPipelineStack.Value;
-								for (int32 PipelineIndex = 0; PipelineIndex < InterchangePipelineStack.Pipelines.Num(); ++PipelineIndex)
+								if (UInterchangePipelineBase* GeneratedPipeline = UE::Interchange::GeneratePipelineInstance(DefaultPipeline))
 								{
-									if (UInterchangePipelineBase* GeneratedPipeline = UE::Interchange::GeneratePipelineInstance(InterchangePipelineStack.Pipelines[PipelineIndex]))
-									{
-										if (GeneratedPipeline->GetClass() == PipelineClass)
-										{
-											NodeItem->Pipeline = GeneratedPipeline;
-											GeneratedPipeline->AdjustSettingsForContext((bSceneImport ? EInterchangePipelineContext::SceneImport : EInterchangePipelineContext::AssetImport), nullptr);
-											PipelineConfigurationDetailsView->SetObject(NodeItem->Pipeline, true);
-											//Exit the lambda
-											return;
-										}
-									}
+									GeneratedPipeline->TransferAdjustSettings(DefaultPipeline);
+									//Switch the pipeline the element point on
+									PipelineListViewItems.RemoveAt(PipelineIndex, 1, false);
+									PipelineListViewItems.Insert(GeneratedPipeline, PipelineIndex);
+									PipelineConfigurationDetailsView->SetObject(GeneratedPipeline, true);
+									break;
 								}
 							}
 						}
-					});
+					}
 				}
 			}
 		}
@@ -713,18 +453,12 @@ FReply SInterchangePipelineConfigurationDialog::OnResetToDefault()
 	return Result;
 }
 
-bool SInterchangePipelineConfigurationDialog::RecursiveValidatePipelineSettings(const TSharedPtr<FInterchangePipelineStacksTreeNodeItem>& ParentNode, TOptional<FText>& OutInvalidReason) const
+bool SInterchangePipelineConfigurationDialog::ValidateAllPipelineSettings(TOptional<FText>& OutInvalidReason) const
 {
-	if (ParentNode->Pipeline)
-	{
-		if (!ParentNode->Pipeline->IsSettingsAreValid(OutInvalidReason))
-		{
-			return false;
-		}
-	}
-	for (int32 ChildIndex = 0; ChildIndex < ParentNode->Childrens.Num(); ++ChildIndex)
+	for (TObjectPtr<UInterchangePipelineBase> PipelineElement : PipelineListViewItems)
 	{
-		if (!RecursiveValidatePipelineSettings(ParentNode->Childrens[ChildIndex], OutInvalidReason))
+		check(PipelineElement);
+		if (!PipelineElement->IsSettingsAreValid(OutInvalidReason))
 		{
 			return false;
 		}
@@ -734,60 +468,28 @@ bool SInterchangePipelineConfigurationDialog::RecursiveValidatePipelineSettings(
 
 bool SInterchangePipelineConfigurationDialog::IsImportButtonEnabled() const
 {
-	const FName DefaultPipelineStackName = FInterchangeProjectSettingsUtils::GetDefaultPipelineStackName(bSceneImport, *SourceData.Get());
-	const TArray<TSharedPtr<FInterchangePipelineStacksTreeNodeItem>>& RootNodeArray = PipelineConfigurationTreeView->GetRootNodeArray();
-	for (const TSharedPtr<FInterchangePipelineStacksTreeNodeItem>& RootNode : RootNodeArray)
-	{
-		//Reimport has only one stack and the name wont match DefaultPipelineStackName
-		if (bReimport || RootNode->StackName == DefaultPipelineStackName)
-		{
-			TOptional<FText> InvalidReason;
-			return RecursiveValidatePipelineSettings(RootNode, InvalidReason);
-		}
-	}
-	return true;
+	TOptional<FText> InvalidReason;
+	return ValidateAllPipelineSettings(InvalidReason);
 }
 
 FText SInterchangePipelineConfigurationDialog::GetImportButtonTooltip() const
 {
-	const FName DefaultPipelineStackName = FInterchangeProjectSettingsUtils::GetDefaultPipelineStackName(bSceneImport, *SourceData.Get());
-	const TArray<TSharedPtr<FInterchangePipelineStacksTreeNodeItem>>& RootNodeArray = PipelineConfigurationTreeView->GetRootNodeArray();
-	for (const TSharedPtr<FInterchangePipelineStacksTreeNodeItem>& RootNode : RootNodeArray)
+	TOptional<FText> InvalidReason;
+	if (!ValidateAllPipelineSettings(InvalidReason) && InvalidReason.IsSet())
 	{
-		//Reimport have only one stack and the name wont match DefaultPipelineStackName
-		if (bReimport || RootNode->StackName == DefaultPipelineStackName)
-		{
-			TOptional<FText> InvalidReason;
-			if (!RecursiveValidatePipelineSettings(RootNode, InvalidReason) && InvalidReason.IsSet())
-			{
-				return InvalidReason.GetValue();
-			}
-		}
+		return InvalidReason.GetValue();
 	}
 	return FText();
 }
 
-void SInterchangePipelineConfigurationDialog::RecursiveSavePipelineSettings(const TSharedPtr<FInterchangePipelineStacksTreeNodeItem>& ParentNode, const int32 PipelineIndex) const
-{
-	if (ParentNode->Pipeline)
-	{
-		ParentNode->Pipeline->SaveSettings(ParentNode->StackName);
-	}
-	for (int32 ChildIndex = 0; ChildIndex < ParentNode->Childrens.Num(); ++ChildIndex)
-	{
-		RecursiveSavePipelineSettings(ParentNode->Childrens[ChildIndex], ChildIndex);
-	}
-}
-
-void SInterchangePipelineConfigurationDialog::RecursiveLoadPipelineSettings(const TSharedPtr<FInterchangePipelineStacksTreeNodeItem>& ParentNode, const int32 PipelineIndex) const
+void SInterchangePipelineConfigurationDialog::SaveAllPipelineSettings() const
 {
-	if (ParentNode->Pipeline)
+	for (const TObjectPtr<UInterchangePipelineBase> PipelineElement : PipelineListViewItems)
 	{
-		ParentNode->Pipeline->LoadSettings(ParentNode->StackName);
-	}
-	for (int32 ChildIndex = 0; ChildIndex < ParentNode->Childrens.Num(); ++ChildIndex)
-	{
-		RecursiveLoadPipelineSettings(ParentNode->Childrens[ChildIndex], ChildIndex);
+		if (PipelineElement)
+		{
+			PipelineElement->SaveSettings(CurrentStackName);
+		}
 	}
 }
 
@@ -797,32 +499,25 @@ void SInterchangePipelineConfigurationDialog::ClosePipelineConfiguration(const E
 	{
 		bCanceled = true;
 		bImportAll = false;
-		if (bReimport && PipelineConfigurationTreeView)
-		{
-			//If user cancel the reimport we have to put back the pipeline settings
-			const TArray<TSharedPtr<FInterchangePipelineStacksTreeNodeItem>>& RootNodeArray = PipelineConfigurationTreeView->GetRootNodeArray();
-			for (const TSharedPtr<FInterchangePipelineStacksTreeNodeItem>& RootNode : RootNodeArray)
-			{
-				RecursiveLoadPipelineSettings(RootNode, 0);
-			}
-		}
 	}
 	else //ECloseEventType::Import
 	{
 		bCanceled = false;
 		bImportAll = UseSameSettingsForAllCheckBox->IsChecked();
+		
+		//Fill the OutPipelines array
+		for (const TObjectPtr<UInterchangePipelineBase> PipelineElement : PipelineListViewItems)
+		{
+			OutPipelines->Add(PipelineElement.Get());
+		}
 	}
 
-	if (!bReimport && PipelineConfigurationTreeView)
+	//Save the settings only if its not a reimport
+	if (!bReimport)
 	{
-		const TArray<TSharedPtr<FInterchangePipelineStacksTreeNodeItem>>& RootNodeArray = PipelineConfigurationTreeView->GetRootNodeArray();
-		for (const TSharedPtr<FInterchangePipelineStacksTreeNodeItem>& RootNode : RootNodeArray)
-		{
-			RecursiveSavePipelineSettings(RootNode, 0);
-		}
+		SaveAllPipelineSettings();
 	}
 
-	PipelineConfigurationTreeView = nullptr;
 	PipelineConfigurationDetailsView = nullptr;
 
 	if (CloseEventType != ECloseEventType::WindowClosing)
@@ -852,4 +547,89 @@ FReply SInterchangePipelineConfigurationDialog::OnKeyDown(const FGeometry& MyGeo
 	return FReply::Unhandled();
 }
 
+void SInterchangePipelineConfigurationDialog::OnStackSelectionChanged(TSharedPtr<FString> String, ESelectInfo::Type)
+{
+	if (!String.IsValid())
+	{
+		return;
+	}
+
+	FName NewStackName = FName(*String.Get());
+	if (!UE::Private::ContainStack(PipelineStacks, NewStackName))
+	{
+		return;
+	}
+
+	//Nothing change the selection is the same
+	if (CurrentStackName == NewStackName)
+	{
+		return;
+	}
+
+	//Save current stack settings, we want the same settings when we will go back to the same stack
+	//When doing a reimport we do not want to save the setting because the context have special default
+	//value for some options like: (Import Materials, Import Textures...).
+	//So when doing a reimport switching stack is like doing a reset to default on all pipelines
+	if (!bReimport)
+	{
+		SaveAllPipelineSettings();
+	}
+
+	//Use the stack select by interchange manager
+	CurrentStackName = NewStackName;
+	
+	//Rebuild the Pipeline list item
+	PipelineListViewItems.Reset();
+
+	for (FInterchangeStackInfo& Stack : PipelineStacks)
+	{
+		TSharedPtr<FString> StackNamePtr = MakeShared<FString>(Stack.StackName.ToString());
+		if (CurrentStackName != Stack.StackName)
+		{
+			continue;
+		}
+		for (UInterchangePipelineBase* DefaultPipeline : Stack.Pipelines)
+		{
+			check(DefaultPipeline);
+			if (UInterchangePipelineBase* GeneratedPipeline = UE::Interchange::GeneratePipelineInstance(DefaultPipeline))
+			{
+				GeneratedPipeline->TransferAdjustSettings(DefaultPipeline);
+				if (Stack.StackName != ReimportStackName)
+				{
+					//Load the settings for this pipeline
+					GeneratedPipeline->LoadSettings(Stack.StackName);
+					GeneratedPipeline->PreDialogCleanup(Stack.StackName);
+				}
+				PipelineListViewItems.Add(GeneratedPipeline);
+			}
+		}
+	}
+
+	//Select the first pipeline
+	if (PipelineListViewItems.Num() > 0)
+	{
+		PipelinesListView->SetSelection(PipelineListViewItems[0], ESelectInfo::Direct);
+	}
+	PipelinesListView->RequestListRefresh();
+}
+
+
+TSharedRef<ITableRow> SInterchangePipelineConfigurationDialog::MakePipelineListRowWidget(
+	TObjectPtr<UInterchangePipelineBase> InElement,
+	const TSharedRef<STableViewBase>& OwnerTable)
+{
+	check(InElement);
+	return SNew(SInterchangePipelineItem, OwnerTable, InElement);
+}
+
+void SInterchangePipelineConfigurationDialog::OnPipelineSelectionChanged(TObjectPtr<UInterchangePipelineBase> InItem, ESelectInfo::Type SelectInfo)
+{
+	CurrentSelectedPipeline = nullptr;
+	if (InItem)
+	{
+		CurrentSelectedPipeline = InItem;
+	}
+	PipelineConfigurationDetailsView->SetObject(CurrentSelectedPipeline.Get());
+}
+
 #undef LOCTEXT_NAMESPACE
diff --git a/Engine/Plugins/Interchange/Editor/Source/Pipelines/Private/SInterchangePipelineConfigurationDialog.h b/Engine/Plugins/Interchange/Editor/Source/Pipelines/Private/SInterchangePipelineConfigurationDialog.h
--- a/Engine/Plugins/Interchange/Editor/Source/Pipelines/Private/SInterchangePipelineConfigurationDialog.h
+++ b/Engine/Plugins/Interchange/Editor/Source/Pipelines/Private/SInterchangePipelineConfigurationDialog.h
@@ -5,6 +5,7 @@
 #include "InputCoreTypes.h"
 #include "Input/Reply.h"
 #include "InterchangePipelineBase.h"
+#include "InterchangePipelineConfigurationBase.h"
 #include "InterchangeSourceData.h"
 #include "Styling/SlateColor.h"
 #include "UObject/GCObject.h"
@@ -17,79 +18,18 @@ struct FPropertyAndParent;
 class IDetailsView;
 class SCheckBox;
 
-class FInterchangePipelineStacksTreeNodeItem : protected FGCObject
+class SInterchangePipelineItem : public STableRow<TObjectPtr<UInterchangePipelineBase>>
 {
 public:
-	FInterchangePipelineStacksTreeNodeItem()
-	{
-		StackName = NAME_None;
-		Pipeline = nullptr;
-	}
-
-	/* FGCObject interface */
-	virtual void AddReferencedObjects(FReferenceCollector& Collector) override;
-	virtual FString GetReferencerName() const override
-	{
-		return TEXT("FInterchangePipelineStacksTreeNodeItem");
-	}
-
-	//This name is use only when this item represent a stack name
-	FName StackName;
-
-	//Pipeline is nullptr when the node represent a stack name
-	UInterchangePipelineBase* Pipeline;
-
-	TArray<TSharedPtr<FInterchangePipelineStacksTreeNodeItem>> Childrens;
-};
-
-DECLARE_DELEGATE_TwoParams(FOnPipelineConfigurationSelectionChanged, TSharedPtr<FInterchangePipelineStacksTreeNodeItem>, ESelectInfo::Type)
-
-class SInterchangePipelineStacksTreeView : public STreeView< TSharedPtr<FInterchangePipelineStacksTreeNodeItem>>
-{
-public:
-	~SInterchangePipelineStacksTreeView();
-
-	SLATE_BEGIN_ARGS(SInterchangePipelineStacksTreeView)
-		: _OnSelectionChangedDelegate()
-	{}
-		SLATE_EVENT(FOnPipelineConfigurationSelectionChanged, OnSelectionChangedDelegate)
-		SLATE_ARGUMENT(TWeakObjectPtr<UInterchangeSourceData>, SourceData)
-		SLATE_ARGUMENT(bool, bSceneImport)
-		SLATE_ARGUMENT(bool, bReimport)
-		SLATE_ARGUMENT(TArray<UInterchangePipelineBase*>, PipelineStack)
-	SLATE_END_ARGS()
-
-	/** Construct this widget */
-	void Construct(const FArguments& InArgs);
-	TSharedRef< ITableRow > OnGenerateRowPipelineConfigurationTreeView(TSharedPtr<FInterchangePipelineStacksTreeNodeItem> Item, const TSharedRef< STableViewBase >& OwnerTable);
-	void OnGetChildrenPipelineConfigurationTreeView(TSharedPtr<FInterchangePipelineStacksTreeNodeItem> InParent, TArray< TSharedPtr<FInterchangePipelineStacksTreeNodeItem> >& OutChildren);
-
-	FReply OnExpandAll();
-	FReply OnCollapseAll();
-
-	const TArray<TSharedPtr<FInterchangePipelineStacksTreeNodeItem>>& GetRootNodeArray() const { return RootNodeArray; }
-	TArray<TSharedPtr<FInterchangePipelineStacksTreeNodeItem>>& GetMutableRootNodeArray() { return RootNodeArray; }
-
-	void SelectDefaultItem();
-
-protected:
-	/** Delegate to invoke when selection changes. */
-	FOnPipelineConfigurationSelectionChanged OnSelectionChangedDelegate;
-	TWeakObjectPtr<UInterchangeSourceData> SourceData;
-	bool bSceneImport = false;
-	bool bReimport = false;
-	TArray<UInterchangePipelineBase*> PipelineStack;
-
-	/** the elements we show in the tree view */
-	TArray<TSharedPtr<FInterchangePipelineStacksTreeNodeItem>> RootNodeArray;
-
-	/** Open a context menu for the current selection */
-	TSharedPtr<SWidget> OnOpenContextMenu();
-	void SetAsDefaultStack(FName NewDefaultStackValue);
-	void RecursiveSetExpand(TSharedPtr<FInterchangePipelineStacksTreeNodeItem> Node, bool ExpandState);
-	void OnTreeViewSelectionChanged(TSharedPtr<FInterchangePipelineStacksTreeNodeItem> Item, ESelectInfo::Type SelectionType);
+	void Construct(
+		const FArguments& InArgs,
+		const TSharedRef<STableViewBase>& OwnerTable,
+		TObjectPtr<UInterchangePipelineBase> InPipelineElement);
+private:
+	TObjectPtr<UInterchangePipelineBase> PipelineElement;
 };
 
+typedef SListView< TObjectPtr<UInterchangePipelineBase> > SPipelineListViewType;
 
 enum class ECloseEventType : uint8
 {
@@ -109,7 +49,8 @@ public:
 		SLATE_ARGUMENT(TWeakObjectPtr<UInterchangeSourceData>, SourceData)
 		SLATE_ARGUMENT(bool, bSceneImport)
 		SLATE_ARGUMENT(bool, bReimport)
-		SLATE_ARGUMENT(TArray<UInterchangePipelineBase*>, PipelineStack)
+		SLATE_ARGUMENT(TArray<FInterchangeStackInfo>, PipelineStacks)
+		SLATE_ARGUMENT(TArray<UInterchangePipelineBase*>*, OutPipelines)
 	SLATE_END_ARGS()
 
 public:
@@ -140,27 +81,41 @@ public:
 
 private:
 	TSharedRef<SBox> SpawnPipelineConfiguration();
-	void OnSelectionChanged(TSharedPtr<FInterchangePipelineStacksTreeNodeItem> Item, ESelectInfo::Type SelectionType);
 
 	bool IsPropertyVisible(const FPropertyAndParent&) const;
 	FText GetSourceDescription() const;
-	void RecursiveIterateNode(TSharedPtr<FInterchangePipelineStacksTreeNodeItem>& ParentNode, TFunctionRef<void(TSharedPtr<FInterchangePipelineStacksTreeNodeItem>& CurrentNode)> IterationLambda);
 	FReply OnResetToDefault();
 
-	bool RecursiveValidatePipelineSettings(const TSharedPtr<FInterchangePipelineStacksTreeNodeItem>& ParentNode, TOptional<FText>& OutInvalidReason) const;
+	bool ValidateAllPipelineSettings(TOptional<FText>& OutInvalidReason) const;
 	bool IsImportButtonEnabled() const;
 	FText GetImportButtonTooltip() const;
 
-	void RecursiveSavePipelineSettings(const TSharedPtr<FInterchangePipelineStacksTreeNodeItem>& ParentNode, const int32 PipelineIndex) const;
-	void RecursiveLoadPipelineSettings(const TSharedPtr<FInterchangePipelineStacksTreeNodeItem>& ParentNode, const int32 PipelineIndex) const;
+	void SaveAllPipelineSettings() const;
 
 private:
 	TWeakPtr< SWindow > OwnerWindow;
 	TWeakObjectPtr<UInterchangeSourceData> SourceData;
-	TArray<UInterchangePipelineBase*> PipelineStack;
+	TArray<FInterchangeStackInfo> PipelineStacks;
+	TArray<UInterchangePipelineBase*>* OutPipelines;
+
+	// The available stacks
+	TArray<TSharedPtr<FString>> AvailableStacks;
+	void OnStackSelectionChanged(TSharedPtr<FString> String, ESelectInfo::Type);
+
+	//////////////////////////////////////////////////////////////////////////
+	// the pipelines list view
 	
-	//Graph Inspector UI elements
-	TSharedPtr<SInterchangePipelineStacksTreeView> PipelineConfigurationTreeView;
+	TSharedPtr<SPipelineListViewType> PipelinesListView;
+	TArray< TObjectPtr<UInterchangePipelineBase> > PipelineListViewItems;
+
+	/** list view generate row callback */
+	TSharedRef<ITableRow> MakePipelineListRowWidget(TObjectPtr<UInterchangePipelineBase> InElement, const TSharedRef<STableViewBase>& OwnerTable);
+
+	void OnPipelineSelectionChanged(TObjectPtr<UInterchangePipelineBase> InItem, ESelectInfo::Type SelectInfo);
+
+	//
+	//////////////////////////////////////////////////////////////////////////
+
 	TSharedPtr<IDetailsView> PipelineConfigurationDetailsView;
 	TSharedPtr<SCheckBox> UseSameSettingsForAllCheckBox;
 
@@ -170,6 +125,6 @@ private:
 	bool bImportAll = false;
 
 	FName CurrentStackName = NAME_None;
-	TWeakObjectPtr<UInterchangePipelineBase> CurrentSelectedPipeline = nullptr;
+	TObjectPtr<UInterchangePipelineBase> CurrentSelectedPipeline = nullptr;
 };
 
diff --git a/Engine/Plugins/Interchange/Editor/Source/Pipelines/Public/InterchangePipelineConfigurationGeneric.h b/Engine/Plugins/Interchange/Editor/Source/Pipelines/Public/InterchangePipelineConfigurationGeneric.h
--- a/Engine/Plugins/Interchange/Editor/Source/Pipelines/Public/InterchangePipelineConfigurationGeneric.h
+++ b/Engine/Plugins/Interchange/Editor/Source/Pipelines/Public/InterchangePipelineConfigurationGeneric.h
@@ -20,7 +20,13 @@ public:
 
 protected:
 
-	virtual EInterchangePipelineConfigurationDialogResult ShowPipelineConfigurationDialog(TWeakObjectPtr<UInterchangeSourceData> SourceData) override;
-	virtual EInterchangePipelineConfigurationDialogResult ShowScenePipelineConfigurationDialog(TWeakObjectPtr<UInterchangeSourceData> SourceData) override;
-	virtual EInterchangePipelineConfigurationDialogResult ShowReimportPipelineConfigurationDialog(TArray<UInterchangePipelineBase*>& PipelineStack, TWeakObjectPtr<UInterchangeSourceData> SourceData) override;
+	virtual EInterchangePipelineConfigurationDialogResult ShowPipelineConfigurationDialog(TArray<FInterchangeStackInfo>& PipelineStacks
+		, TArray<UInterchangePipelineBase*>& OutPipelines
+		, TWeakObjectPtr<UInterchangeSourceData> SourceData) override;
+	virtual EInterchangePipelineConfigurationDialogResult ShowScenePipelineConfigurationDialog(TArray<FInterchangeStackInfo>& PipelineStacks
+		, TArray<UInterchangePipelineBase*>& OutPipelines
+		, TWeakObjectPtr<UInterchangeSourceData> SourceData) override;
+	virtual EInterchangePipelineConfigurationDialogResult ShowReimportPipelineConfigurationDialog(TArray<FInterchangeStackInfo>& PipelineStacks
+		, TArray<UInterchangePipelineBase*>& OutPipelines
+		, TWeakObjectPtr<UInterchangeSourceData> SourceData) override;
 };
diff --git a/Engine/Source/Runtime/Interchange/Core/Private/InterchangePipelineBase.cpp b/Engine/Source/Runtime/Interchange/Core/Private/InterchangePipelineBase.cpp
--- a/Engine/Source/Runtime/Interchange/Core/Private/InterchangePipelineBase.cpp
+++ b/Engine/Source/Runtime/Interchange/Core/Private/InterchangePipelineBase.cpp
@@ -70,6 +70,15 @@ void UInterchangePipelineBase::AdjustSettingsFromCache()
 	AdjustSettingsForContext(CachePipelineContext, CacheReimportObject.Get());
 }
 
+void UInterchangePipelineBase::TransferAdjustSettings(UInterchangePipelineBase* SourcePipeline)
+{
+	CachePipelineContext = SourcePipeline->CachePipelineContext;
+	CacheReimportObject = SourcePipeline->CacheReimportObject;
+	CachePropertiesStates = SourcePipeline->CachePropertiesStates;
+	bAllowPropertyStatesEdition = SourcePipeline->bAllowPropertyStatesEdition;
+	bIsReimportContext = SourcePipeline->bIsReimportContext;
+}
+
 const FInterchangePipelinePropertyStates* UInterchangePipelineBase::GetPropertyStates(const FName PropertyPath) const
 {
 	return PropertiesStates.Find(PropertyPath);
diff --git a/Engine/Source/Runtime/Interchange/Core/Public/InterchangePipelineBase.h b/Engine/Source/Runtime/Interchange/Core/Public/InterchangePipelineBase.h
--- a/Engine/Source/Runtime/Interchange/Core/Public/InterchangePipelineBase.h
+++ b/Engine/Source/Runtime/Interchange/Core/Public/InterchangePipelineBase.h
@@ -243,6 +243,9 @@ public:
 	virtual void AdjustSettingsForContext(EInterchangePipelineContext ReimportType, TObjectPtr<UObject> ReimportAsset);
 	virtual void AdjustSettingsFromCache();
 
+	/** Transfer the source pipeline adjust settings to this pipeline. */
+	void TransferAdjustSettings(UInterchangePipelineBase* SourcePipeline);
+
 	/**
 	 * This function is used to add the given message object directly into the results for this operation.
 	 */
@@ -296,6 +299,7 @@ public:
 
 	bool CanEditPropertiesStates() { return bAllowPropertyStatesEdition; }
 	bool IsReimportContext() { return bIsReimportContext; }
+
 protected:
 
 	UE_DEPRECATED(5.2, "This function is replace by ExecutePipeline.")
diff --git a/Engine/Source/Runtime/Interchange/Engine/Private/InterchangeManager.cpp b/Engine/Source/Runtime/Interchange/Engine/Private/InterchangeManager.cpp
--- a/Engine/Source/Runtime/Interchange/Engine/Private/InterchangeManager.cpp
+++ b/Engine/Source/Runtime/Interchange/Engine/Private/InterchangeManager.cpp
@@ -1115,6 +1115,22 @@ UInterchangeManager::ImportInternal(const FString& ContentPath, const UInterchan
 	}
 #endif
 
+	const bool bIsReimport = OriginalAssetImportData && OriginalAssetImportData->Pipelines.Num() > 0;
+	auto AdjustPipelineSettingForContext = [bImportScene, bIsReimport, ImportAssetParameters](UInterchangePipelineBase* Pipeline)
+	{
+		EInterchangePipelineContext Context;
+		UObject* ReimportObject = nullptr;
+		if (bIsReimport)
+		{
+			Context = bImportScene ? EInterchangePipelineContext::SceneReimport : EInterchangePipelineContext::AssetReimport;
+			ReimportObject = ImportAssetParameters.ReimportAsset;
+		}
+		else
+		{
+			Context = bImportScene ? EInterchangePipelineContext::SceneImport : EInterchangePipelineContext::AssetImport;
+		}
+		Pipeline->AdjustSettingsForContext(Context, ReimportObject);
+	};
 
 	if ( ImportAssetParameters.OverridePipelines.Num() == 0 )
 	{
@@ -1137,10 +1153,15 @@ UInterchangeManager::ImportInternal(const FString& ContentPath, const UInterchan
 
 		const TMap<FName, FInterchangePipelineStack>& DefaultPipelineStacks = InterchangeImportSettings.PipelineStacks;
 
-		//If we reimport we want to load the original pipeline and the original pipeline settings
-		if (OriginalAssetImportData && OriginalAssetImportData->Pipelines.Num() > 0)
+		const FName ReimportPipelineName = TEXT("ReimportPipeline");
+		TArray<FInterchangeStackInfo> PipelineStacks;
+		TArray<UInterchangePipelineBase*> OutPipelines;
+
+		//Fill the Stacks before showing the UI
+		if (bIsReimport)
 		{
-			TArray<UInterchangePipelineBase*> PipelineStack;
+			FInterchangeStackInfo& StackInfo = PipelineStacks.AddDefaulted_GetRef();
+			StackInfo.StackName = ReimportPipelineName;
 			for (int32 PipelineIndex = 0; PipelineIndex < OriginalAssetImportData->Pipelines.Num(); ++PipelineIndex)
 			{
 				UInterchangePipelineBase* SourcePipeline = Cast<UInterchangePipelineBase>(OriginalAssetImportData->Pipelines[PipelineIndex]);
@@ -1156,16 +1177,9 @@ UInterchangeManager::ImportInternal(const FString& ContentPath, const UInterchan
 					//Duplicate the pipeline saved in the asset import data
 					UInterchangePipelineBase* GeneratedPipeline = Cast<UInterchangePipelineBase>(StaticDuplicateObject(SourcePipeline, GetTransientPackage()));
 					// Make sure that the instance does not carry over standalone and public flags as they are not actual assets to be persisted
-					GeneratedPipeline->ClearFlags(EObjectFlags::RF_Standalone|EObjectFlags::RF_Public);
-					if (bImportScene)
-					{
-						GeneratedPipeline->AdjustSettingsForContext(EInterchangePipelineContext::SceneReimport, nullptr);
-					}
-					else
-					{
-						GeneratedPipeline->AdjustSettingsForContext(EInterchangePipelineContext::AssetReimport, ImportAssetParameters.ReimportAsset);
-					}
-					PipelineStack.Add(GeneratedPipeline);
+					GeneratedPipeline->ClearFlags(EObjectFlags::RF_Standalone | EObjectFlags::RF_Public);
+					AdjustPipelineSettingForContext(GeneratedPipeline);
+					StackInfo.Pipelines.Add(GeneratedPipeline);
 				}
 				else
 				{
@@ -1173,11 +1187,28 @@ UInterchangeManager::ImportInternal(const FString& ContentPath, const UInterchan
 					//Log something
 				}
 			}
+		}
+		for (const TPair<FName, FInterchangePipelineStack>& PipelineStackInfo : DefaultPipelineStacks)
+		{
+			FName StackName = PipelineStackInfo.Key;
+			FInterchangeStackInfo& StackInfo = PipelineStacks.AddDefaulted_GetRef();
+			StackInfo.StackName = StackName;
+			for (int32 PipelineIndex = 0; PipelineIndex < PipelineStackInfo.Value.Pipelines.Num(); ++PipelineIndex)
+			{
+				if (UInterchangePipelineBase* GeneratedPipeline = UE::Interchange::GeneratePipelineInstance(PipelineStackInfo.Value.Pipelines[PipelineIndex]))
+				{
+					AdjustPipelineSettingForContext(GeneratedPipeline);
+					StackInfo.Pipelines.Add(GeneratedPipeline);
+				}
+			}
+		}
 
+		if (bIsReimport)
+		{
 			if (RegisteredPipelineConfiguration && bShowPipelineStacksConfigurationDialog && !bIsUnattended)
 			{
 				//Show the dialog, a plugin should have registered this dialog. We use a plugin to be able to use editor code when doing UI
-				EInterchangePipelineConfigurationDialogResult DialogResult = RegisteredPipelineConfiguration->ScriptedShowReimportPipelineConfigurationDialog(PipelineStack, DuplicateSourceData);
+				EInterchangePipelineConfigurationDialogResult DialogResult = RegisteredPipelineConfiguration->ScriptedShowReimportPipelineConfigurationDialog(PipelineStacks, OutPipelines, DuplicateSourceData);
 				if (DialogResult == EInterchangePipelineConfigurationDialogResult::Cancel)
 				{
 					bImportCanceled = true;
@@ -1187,34 +1218,27 @@ UInterchangeManager::ImportInternal(const FString& ContentPath, const UInterchan
 					bImportAllWithDefault = true;
 				}
 			}
-
-			// Simply move the existing pipeline
-			AsyncHelper->Pipelines = MoveTemp(PipelineStack);
-
-			//Fill the original pipeline array that will be save in the asset import data
-			for (UInterchangePipelineBase* Pipeline : AsyncHelper->Pipelines)
+			else
 			{
-				if (UInterchangePythonPipelineBase* PythonPipeline = Cast<UInterchangePythonPipelineBase>(Pipeline))
-				{
-					UInterchangePythonPipelineAsset* PythonPipelineAsset = NewObject<UInterchangePythonPipelineAsset>(GetTransientPackage());
-					PythonPipelineAsset->PythonClass = PythonPipeline->GetClass();
-					PythonPipelineAsset->SetupFromPipeline(PythonPipeline);
-					AsyncHelper->OriginalPipelines.Add(PythonPipelineAsset);
-				}
-				else
-				{
-					AsyncHelper->OriginalPipelines.Add(Pipeline);
-				}
-				UE::Interchange::Private::FillPipelineAnalyticData(Pipeline, UniqueId, FString());
+				//When we do not show the UI we use the original stack
+				FInterchangeStackInfo* StackInfoPtr = PipelineStacks.FindByPredicate([ReimportPipelineName](const FInterchangeStackInfo& StackInfo)
+					{
+						return StackInfo.StackName == ReimportPipelineName;
+					});
+				
+				check(StackInfoPtr);
+				OutPipelines = StackInfoPtr->Pipelines;
 			}
 		}
 		else
 		{
-			FName PipelineStackName = FInterchangeProjectSettingsUtils::GetDefaultPipelineStackName(bImportScene, *SourceData);
-			if (RegisteredPipelineConfiguration && (bShowPipelineStacksConfigurationDialog || (!DefaultPipelineStacks.Contains(PipelineStackName) && !bIsUnattended)))
+			if (RegisteredPipelineConfiguration && bShowPipelineStacksConfigurationDialog)
 			{
 				//Show the dialog, a plugin should have register this dialog. We use a plugin to be able to use editor code when doing UI
-				EInterchangePipelineConfigurationDialogResult DialogResult = bImportScene ? RegisteredPipelineConfiguration->ScriptedShowScenePipelineConfigurationDialog(DuplicateSourceData) : RegisteredPipelineConfiguration->ScriptedShowPipelineConfigurationDialog(DuplicateSourceData);
+				EInterchangePipelineConfigurationDialogResult DialogResult = bImportScene
+					? RegisteredPipelineConfiguration->ScriptedShowScenePipelineConfigurationDialog(PipelineStacks, OutPipelines, DuplicateSourceData)
+					: RegisteredPipelineConfiguration->ScriptedShowPipelineConfigurationDialog(PipelineStacks, OutPipelines, DuplicateSourceData);
+
 				if (DialogResult == EInterchangePipelineConfigurationDialogResult::Cancel)
 				{
 					bImportCanceled = true;
@@ -1224,75 +1248,55 @@ UInterchangeManager::ImportInternal(const FString& ContentPath, const UInterchan
 					bImportAllWithDefault = true;
 				}
 			}
-
-			if (!bImportCanceled)
+			else
 			{
-				if (!DefaultPipelineStacks.Contains(PipelineStackName))
-				{
-					if (DefaultPipelineStacks.Contains(FInterchangeProjectSettingsUtils::GetDefaultPipelineStackName(bImportScene, *SourceData)))
+				FName DefaultStackName = FInterchangeProjectSettingsUtils::GetDefaultPipelineStackName(bImportScene, *DuplicateSourceData);
+				FInterchangeStackInfo* StackInfoPtr = PipelineStacks.FindByPredicate([DefaultStackName](const FInterchangeStackInfo& StackInfo)
 					{
-						PipelineStackName = FInterchangeProjectSettingsUtils::GetDefaultPipelineStackName(bImportScene, *SourceData);
-					}
-					else
+						return StackInfo.StackName == DefaultStackName;
+					});
+				if (StackInfoPtr)
+				{
+					//When we do not show the UI we use the original stack
+					OutPipelines = StackInfoPtr->Pipelines;
+				}
+				else if (PipelineStacks.Num() > 0)
+				{
+					for (FInterchangeStackInfo& StackInfo : PipelineStacks)
 					{
-						//Log an error, we cannot import asset without a valid pipeline, we will use the first available pipeline
-						for (const TPair<FName, FInterchangePipelineStack>& PipelineStack : DefaultPipelineStacks)
-						{
-							PipelineStackName = PipelineStack.Key;
-						}
+						OutPipelines = StackInfo.Pipelines;
+						UE_LOG(LogInterchangeEngine, Warning, TEXT("Interchange import: Invalid Default stack, using stack [%s] to import"), *StackInfo.StackName.ToString());
+						break;
 					}
 				}
+				else
+				{
+					UE_LOG(LogInterchangeEngine, Warning, TEXT("Interchange Import: Cannot find any valid stack, cancelling import."));
+					bImportCanceled = true;
+				}
+			}
+		}
 
-				if (DefaultPipelineStacks.Contains(PipelineStackName))
+		if (!bImportCanceled)
+		{
+			// Simply move the existing pipeline
+			AsyncHelper->Pipelines = MoveTemp(OutPipelines);
+
+			//Fill the original pipeline array that will be save in the asset import data
+			for (UInterchangePipelineBase* Pipeline : AsyncHelper->Pipelines)
+			{
+				if (UInterchangePythonPipelineBase* PythonPipeline = Cast<UInterchangePythonPipelineBase>(Pipeline))
 				{
-					if (FEngineAnalytics::IsAvailable())
-					{
-						Attribs.Add(FAnalyticsEventAttribute(TEXT("PipelineStackName"), PipelineStackName));
-					}
-					//use the default pipeline
-					const FInterchangePipelineStack& PipelineStack = DefaultPipelineStacks.FindChecked(PipelineStackName);
-					const TArray<FSoftObjectPath>& Pipelines = PipelineStack.Pipelines;
-					for (int32 GraphPipelineIndex = 0; GraphPipelineIndex < Pipelines.Num(); ++GraphPipelineIndex)
-					{
-						if (Pipelines[GraphPipelineIndex].IsValid())
-						{
-							if(UInterchangePipelineBase* GeneratedPipeline = UE::Interchange::GeneratePipelineInstance(Pipelines[GraphPipelineIndex]))
-							{
-								GeneratedPipeline->AdjustSettingsForContext(bImportScene ? EInterchangePipelineContext::SceneImport : EInterchangePipelineContext::AssetImport, nullptr);
-
-								if (FInterchangeProjectSettingsUtils::ShouldShowPipelineStacksConfigurationDialog(bImportScene, *SourceData))
-								{
-									//Load the settings for this pipeline
-									//The dialog is saving the settings for all default pipelines
-									GeneratedPipeline->LoadSettings(PipelineStackName);
-								}
-					
-								UE::Interchange::Private::FillPipelineAnalyticData(GeneratedPipeline, UniqueId, FString());
-
-								AsyncHelper->Pipelines.Add(GeneratedPipeline);
-
-								//We need to save the python pipeline asset because we cannot save an asset created with a python class
-								if (UInterchangePythonPipelineAsset* OriginalPipeline = Cast<UInterchangePythonPipelineAsset>(Pipelines[GraphPipelineIndex].TryLoad()))
-								{
-									if (UInterchangePythonPipelineAsset* DuplicatedPythonPipelineAsset = DuplicateObject<UInterchangePythonPipelineAsset>(OriginalPipeline, GetTransientPackage()))
-									{
-										DuplicatedPythonPipelineAsset->SetupFromPipeline(Cast<UInterchangePythonPipelineBase>(GeneratedPipeline));
-										AsyncHelper->OriginalPipelines.Add(DuplicatedPythonPipelineAsset);
-									}
-								}
-								else
-								{
-									AsyncHelper->OriginalPipelines.Add(GeneratedPipeline);
-								}
-							}
-						}
-					}
+					UInterchangePythonPipelineAsset* PythonPipelineAsset = NewObject<UInterchangePythonPipelineAsset>(GetTransientPackage());
+					PythonPipelineAsset->PythonClass = PythonPipeline->GetClass();
+					PythonPipelineAsset->SetupFromPipeline(PythonPipeline);
+					AsyncHelper->OriginalPipelines.Add(PythonPipelineAsset);
 				}
 				else
 				{
-					//Log an error, we cannot import asset without a valid pipeline, there is no pipeline stack defined
-					bImportAborted = true;
+					AsyncHelper->OriginalPipelines.Add(Pipeline);
 				}
+				UE::Interchange::Private::FillPipelineAnalyticData(Pipeline, UniqueId, FString());
 			}
 		}
 	}
@@ -1302,14 +1306,7 @@ UInterchangeManager::ImportInternal(const FString& ContentPath, const UInterchan
 		{
 			// Duplicate the override pipelines to protect the scripted users form making race conditions
 			UInterchangePipelineBase* GeneratedPipeline = DuplicateObject<UInterchangePipelineBase>(ImportAssetParameters.OverridePipelines[GraphPipelineIndex], GetTransientPackage());
-			if (OriginalAssetImportData != nullptr)
-			{
-				GeneratedPipeline->AdjustSettingsForContext(bImportScene ? EInterchangePipelineContext::SceneImport : EInterchangePipelineContext::AssetImport, ImportAssetParameters.ReimportAsset);
-			}
-			else
-			{
-				GeneratedPipeline->AdjustSettingsForContext(bImportScene ? EInterchangePipelineContext::SceneImport : EInterchangePipelineContext::AssetImport, nullptr);
-			}
+			AdjustPipelineSettingForContext(GeneratedPipeline);
 			AsyncHelper->Pipelines.Add(GeneratedPipeline);
 			AsyncHelper->OriginalPipelines.Add(GeneratedPipeline);
 			UE::Interchange::Private::FillPipelineAnalyticData(GeneratedPipeline, UniqueId, FString());
diff --git a/Engine/Source/Runtime/Interchange/Engine/Private/InterchangeProjectSettings.cpp b/Engine/Source/Runtime/Interchange/Engine/Private/InterchangeProjectSettings.cpp
--- a/Engine/Source/Runtime/Interchange/Engine/Private/InterchangeProjectSettings.cpp
+++ b/Engine/Source/Runtime/Interchange/Engine/Private/InterchangeProjectSettings.cpp
@@ -67,6 +67,45 @@ FName FInterchangeProjectSettingsUtils::GetDefaultPipelineStackName(const bool b
 	return DefaultPipelineStack;
 }
 
+void FInterchangeProjectSettingsUtils::SetDefaultPipelineStackName(const bool bIsSceneImport, const UInterchangeSourceData& SourceData, const FName StackName)
+{
+	FInterchangeImportSettings& ImportSettings = GetMutableDefaultImportSettings(bIsSceneImport);
+
+	if (!ImportSettings.PipelineStacks.Contains(StackName))
+	{
+		//The new stack name must be valid
+		return;
+	}
+
+	FName DefaultPipelineStack = ImportSettings.DefaultPipelineStack;
+
+	if (!bIsSceneImport)
+	{
+		UE::Interchange::FScopedTranslator ScopedTranslator(&SourceData);
+
+		if (UInterchangeTranslatorBase* Translator = ScopedTranslator.GetTranslator())
+		{
+			EInterchangeTranslatorAssetType SupportedAssetTypes = Translator->GetSupportedAssetTypes();
+
+			FInterchangeContentImportSettings& ContentImportSettings = GetMutableDefault<UInterchangeProjectSettings>()->ContentImportSettings;
+			for (TMap<EInterchangeTranslatorAssetType, FName>::TIterator StackOverridesIt = ContentImportSettings.DefaultPipelineStackOverride.CreateIterator(); StackOverridesIt; ++StackOverridesIt)
+			{
+				if ((SupportedAssetTypes ^ StackOverridesIt->Key) < StackOverridesIt->Key)
+				{
+					//Update the override stack name and save the config
+					StackOverridesIt->Value = StackName;
+					GetMutableDefault<UInterchangeProjectSettings>()->SaveConfig();
+					return;
+				}
+			}
+		}
+	}
+
+	//We do not have any override default stack, simply change the DefaultPipelineStack and save the config
+	ImportSettings.DefaultPipelineStack = StackName;
+	GetMutableDefault<UInterchangeProjectSettings>()->SaveConfig();
+}
+
 bool FInterchangeProjectSettingsUtils::ShouldShowPipelineStacksConfigurationDialog(const bool bIsSceneImport, const UInterchangeSourceData& SourceData)
 {
 	const FInterchangeImportSettings& ImportSettings = GetDefaultImportSettings(bIsSceneImport);
diff --git a/Engine/Source/Runtime/Interchange/Engine/Public/InterchangePipelineConfigurationBase.h b/Engine/Source/Runtime/Interchange/Engine/Public/InterchangePipelineConfigurationBase.h
--- a/Engine/Source/Runtime/Interchange/Engine/Public/InterchangePipelineConfigurationBase.h
+++ b/Engine/Source/Runtime/Interchange/Engine/Public/InterchangePipelineConfigurationBase.h
@@ -17,6 +17,18 @@ enum class EInterchangePipelineConfigurationDialogResult : uint8
 	ImportAll	UMETA(DisplayName = "Import All"),
 };
 
+USTRUCT(BlueprintType)
+struct FInterchangeStackInfo
+{
+	GENERATED_BODY()
+
+	UPROPERTY(EditAnywhere, BlueprintReadWrite, Category = "Interchange | Translator")
+	FName StackName;
+
+	UPROPERTY(EditAnywhere, BlueprintReadWrite, Category = "Interchange | Translator")
+	TArray<TObjectPtr<UInterchangePipelineBase>> Pipelines;
+};
+
 UCLASS(BlueprintType, Blueprintable, Experimental)
 class INTERCHANGEENGINE_API UInterchangePipelineConfigurationBase : public UObject
 {
@@ -28,39 +40,51 @@ public:
 	 * Non virtual helper to allow blueprint to implement event base function to implement the ShowPipelineConfigurationDialog,
 	 */
 	UFUNCTION(BlueprintNativeEvent, BlueprintCallable, Category = "Interchange | Translator")
-	EInterchangePipelineConfigurationDialogResult ScriptedShowPipelineConfigurationDialog(UInterchangeSourceData* SourceData);
+	EInterchangePipelineConfigurationDialogResult ScriptedShowPipelineConfigurationDialog(TArray<FInterchangeStackInfo>& PipelineStacks
+		, TArray<UInterchangePipelineBase*>& OutPipelines
+		, UInterchangeSourceData* SourceData);
 
 	/** The default implementation (call if the blueprint do not have any implementation) will call the virtual ExecuteImportPipeline */
-	EInterchangePipelineConfigurationDialogResult ScriptedShowPipelineConfigurationDialog_Implementation(UInterchangeSourceData* SourceData)
+	EInterchangePipelineConfigurationDialogResult ScriptedShowPipelineConfigurationDialog_Implementation(TArray<FInterchangeStackInfo>& PipelineStacks
+		, TArray<UInterchangePipelineBase*>& OutPipelines
+		, UInterchangeSourceData* SourceData)
 	{
 		//By default we call the virtual import pipeline execution
-		return ShowPipelineConfigurationDialog(SourceData);
+		return ShowPipelineConfigurationDialog(PipelineStacks, OutPipelines, SourceData);
 	}
 
 	/**
 	 * Non virtual helper to allow blueprint to implement event base function to implement the ShowPipelineConfigurationDialog,
 	 */
 	UFUNCTION(BlueprintNativeEvent, BlueprintCallable, Category = "Interchange | Translator")
-	EInterchangePipelineConfigurationDialogResult ScriptedShowScenePipelineConfigurationDialog(UInterchangeSourceData* SourceData);
+	EInterchangePipelineConfigurationDialogResult ScriptedShowScenePipelineConfigurationDialog(TArray<FInterchangeStackInfo>& PipelineStacks
+		, TArray<UInterchangePipelineBase*>& OutPipelines
+		, UInterchangeSourceData* SourceData);
 
 	/** The default implementation (call if the blueprint do not have any implementation) will call the virtual ExecuteImportPipeline */
-	EInterchangePipelineConfigurationDialogResult ScriptedShowScenePipelineConfigurationDialog_Implementation(UInterchangeSourceData* SourceData)
+	EInterchangePipelineConfigurationDialogResult ScriptedShowScenePipelineConfigurationDialog_Implementation(TArray<FInterchangeStackInfo>& PipelineStacks
+		, TArray<UInterchangePipelineBase*>& OutPipelines
+		, UInterchangeSourceData* SourceData)
 	{
 		//By default we call the virtual import pipeline execution
-		return ShowScenePipelineConfigurationDialog(SourceData);
+		return ShowScenePipelineConfigurationDialog(PipelineStacks, OutPipelines, SourceData);
 	}
 
 	/**
 	 * Non virtual helper to allow blueprint to implement event base function to implement the ShowPipelineConfigurationDialog,
 	 */
 	UFUNCTION(BlueprintNativeEvent, BlueprintCallable, Category = "Interchange | Translator")
-	EInterchangePipelineConfigurationDialogResult ScriptedShowReimportPipelineConfigurationDialog(TArray<UInterchangePipelineBase*>& PipelineStack, UInterchangeSourceData* SourceData);
+	EInterchangePipelineConfigurationDialogResult ScriptedShowReimportPipelineConfigurationDialog(TArray<FInterchangeStackInfo>& PipelineStacks
+		, TArray<UInterchangePipelineBase*>& OutPipelines
+		, UInterchangeSourceData* SourceData);
 
 	/** The default implementation (call if the blueprint do not have any implementation) will call the virtual ExecuteImportPipeline */
-	EInterchangePipelineConfigurationDialogResult ScriptedShowReimportPipelineConfigurationDialog_Implementation(TArray<UInterchangePipelineBase*>& PipelineStack, UInterchangeSourceData* SourceData)
+	EInterchangePipelineConfigurationDialogResult ScriptedShowReimportPipelineConfigurationDialog_Implementation(TArray<FInterchangeStackInfo>& PipelineStacks
+		, TArray<UInterchangePipelineBase*>& OutPipelines
+		, UInterchangeSourceData* SourceData)
 	{
 		//By default we call the virtual import pipeline execution
-		return ShowReimportPipelineConfigurationDialog(PipelineStack, SourceData);
+		return ShowReimportPipelineConfigurationDialog(PipelineStacks, OutPipelines, SourceData);
 	}
 
 protected:
@@ -68,7 +92,9 @@ protected:
 	/**
 	 * This function show a dialog use to configure pipeline stacks and return a stack name that tell the caller the user choice.
 	 */
-	virtual EInterchangePipelineConfigurationDialogResult ShowPipelineConfigurationDialog(TWeakObjectPtr<UInterchangeSourceData> SourceData)
+	virtual EInterchangePipelineConfigurationDialogResult ShowPipelineConfigurationDialog(TArray<FInterchangeStackInfo>& PipelineStacks
+		, TArray<UInterchangePipelineBase*>& OutPipelines
+		, TWeakObjectPtr<UInterchangeSourceData> SourceData)
 	{ 
 		//Not implemented
 		return EInterchangePipelineConfigurationDialogResult::Cancel;
@@ -77,7 +103,9 @@ protected:
 	/**
 	 * This function show a dialog use to configure pipeline stacks and return a stack name that tell the caller the user choice.
 	 */
-	virtual EInterchangePipelineConfigurationDialogResult ShowScenePipelineConfigurationDialog(TWeakObjectPtr<UInterchangeSourceData> SourceData)
+	virtual EInterchangePipelineConfigurationDialogResult ShowScenePipelineConfigurationDialog(TArray<FInterchangeStackInfo>& PipelineStacks
+		, TArray<UInterchangePipelineBase*>& OutPipelines
+		, TWeakObjectPtr<UInterchangeSourceData> SourceData)
 	{ 
 		//Not implemented
 		return EInterchangePipelineConfigurationDialogResult::Cancel;
@@ -86,7 +114,9 @@ protected:
 	/**
 	 * This function show a dialog use to configure pipeline stacks and return a stack name that tell the caller the user choice.
 	 */
-	virtual EInterchangePipelineConfigurationDialogResult ShowReimportPipelineConfigurationDialog(TArray<UInterchangePipelineBase*>& PipelineStack, TWeakObjectPtr<UInterchangeSourceData> SourceData)
+	virtual EInterchangePipelineConfigurationDialogResult ShowReimportPipelineConfigurationDialog(TArray<FInterchangeStackInfo>& PipelineStacks
+		, TArray<UInterchangePipelineBase*>& OutPipelines
+		, TWeakObjectPtr<UInterchangeSourceData> SourceData)
 	{
 		//Not implemented
 		return EInterchangePipelineConfigurationDialogResult::Cancel;
diff --git a/Engine/Source/Runtime/Interchange/Engine/Public/InterchangeProjectSettings.h b/Engine/Source/Runtime/Interchange/Engine/Public/InterchangeProjectSettings.h
--- a/Engine/Source/Runtime/Interchange/Engine/Public/InterchangeProjectSettings.h
+++ b/Engine/Source/Runtime/Interchange/Engine/Public/InterchangeProjectSettings.h
@@ -105,6 +105,7 @@ public:
 	static FInterchangeImportSettings& GetMutableDefaultImportSettings(const bool bIsSceneImport);
 
 	static FName GetDefaultPipelineStackName(const bool bIsSceneImport, const UInterchangeSourceData& SourceData);
+	static void SetDefaultPipelineStackName(const bool bIsSceneImport, const UInterchangeSourceData& SourceData, const FName StackName);
 
 	static bool ShouldShowPipelineStacksConfigurationDialog(const bool bIsSceneImport, const UInterchangeSourceData& SourceData);
 };
\ No newline at end of file
