Interchange: Import ui now have the correct tooltip. It also recover the last selected pipeline when the dialog was last shown.

#rnx

diff --git a/Engine/Plugins/Interchange/Editor/Source/Pipelines/Private/SInterchangePipelineConfigurationDialog.cpp b/Engine/Plugins/Interchange/Editor/Source/Pipelines/Private/SInterchangePipelineConfigurationDialog.cpp
--- a/Engine/Plugins/Interchange/Editor/Source/Pipelines/Private/SInterchangePipelineConfigurationDialog.cpp
+++ b/Engine/Plugins/Interchange/Editor/Source/Pipelines/Private/SInterchangePipelineConfigurationDialog.cpp
@@ -14,6 +14,7 @@
 #include "InterchangePythonPipelineBase.h"
 #include "InterchangeSourceData.h"
 #include "Misc/App.h"
+#include "Misc/ConfigCacheIni.h"
 #include "Modules/ModuleManager.h"
 #include "PropertyEditorDelegates.h"
 #include "PropertyEditorModule.h"
@@ -46,8 +47,18 @@ void SInterchangePipelineItem::Construct(
 		STableRow<TObjectPtr<UInterchangePipelineBase>>::FArguments()
 		.Content()
 		[
-			SNew(SBox)
+			SNew(SHorizontalBox)
+			+ SHorizontalBox::Slot()
+			.AutoWidth()
+			.Padding(0.0f, 2.0f, 6.0f, 2.0f)
+			[
+				SNew(SImage)
+				.Image(this, &SInterchangePipelineItem::GetImageItemIcon)
+			]
+			+ SHorizontalBox::Slot()
+			.FillWidth(1.0f)
 			.Padding(3.0f, 0.0f)
+			.VAlign(VAlign_Center)
 			[
 				SNew(STextBlock)
 				.Text(PipelineName)
@@ -55,6 +66,19 @@ void SInterchangePipelineItem::Construct(
 		], OwnerTable);
 }
 
+const FSlateBrush* SInterchangePipelineItem::GetImageItemIcon() const
+{
+	const FSlateBrush* TypeIcon = nullptr;
+	FName IconName = "PipelineConfigurationIcon.Pipeline";
+	const FSlateIcon SlateIcon = FSlateIconFinder::FindIcon(IconName);
+	TypeIcon = SlateIcon.GetOptionalIcon();
+	if (!TypeIcon)
+	{
+		TypeIcon = FSlateIconFinder::FindIconBrushForClass(AActor::StaticClass());
+	}
+	return TypeIcon;
+}
+
 /************************************************************************/
 /* SInterchangePipelineConfigurationDialog Implementation                      */
 /************************************************************************/
@@ -135,32 +159,35 @@ TSharedRef<SBox> SInterchangePipelineConfigurationDialog::SpawnPipelineConfigura
 		AvailableStacks.Add(StackNamePtr);
 	}
 
-
+	FText PipelineListTooltip = LOCTEXT("PipelineListTooltip", "Select a pipeline you want to edit properties for. The pipeline properties will be recorded and changes will be available in subsequent use of that pipeline");
 	PipelinesListView = SNew(SPipelineListViewType)
 		.SelectionMode(ESelectionMode::Single)
 		.ListItemsSource(&PipelineListViewItems)
 		.OnGenerateRow(this, &SInterchangePipelineConfigurationDialog::MakePipelineListRowWidget)
-		.OnSelectionChanged(this, &SInterchangePipelineConfigurationDialog::OnPipelineSelectionChanged);
+		.OnSelectionChanged(this, &SInterchangePipelineConfigurationDialog::OnPipelineSelectionChanged)
+		.ToolTipText(PipelineListTooltip);
 
 	TSharedPtr<STextComboBox> TextComboBoxPtr;
 	//Only use a combo box if there is more then one stack
 	if (AvailableStacks.Num() > 1)
 	{
+		FText StackComboBoxTooltip = LOCTEXT("StackComboBoxTooltip", "Selected pipeline stack will the used for the current import. To change the pipeline stack used when automating or without dialog please change the default pipeline stacks in the project settings");
 		TextComboBoxPtr = SNew(STextComboBox)
 			.OptionsSource(&AvailableStacks)
-			.OnSelectionChanged(this, &SInterchangePipelineConfigurationDialog::OnStackSelectionChanged);
+			.OnSelectionChanged(this, &SInterchangePipelineConfigurationDialog::OnStackSelectionChanged)
+			.ToolTipText(StackComboBoxTooltip);
 		if (SelectedStack.IsValid())
 		{
 			TextComboBoxPtr->SetSelectedItem(SelectedStack);
 		}
 	}
 
-	FText CurrentStackText = LOCTEXT("CurrentStackText", "Stack: ");
+	FText CurrentStackText = LOCTEXT("CurrentStackText", "Choose Pipeline Stack: ");
 
 	TSharedPtr<SWidget> StackTextComboBox;
 	if (!TextComboBoxPtr.IsValid())
 	{
-		CurrentStackText = FText::Format(LOCTEXT("CurrentStackTextNoComboBox", "Stack: {0}"), FText::FromName(CurrentStackName));
+		CurrentStackText = FText::Format(LOCTEXT("CurrentStackTextNoComboBox", "Pipeline Stack: {0}"), FText::FromName(CurrentStackName));
 		StackTextComboBox = SNew(SBox)
 		[
 			SNew(STextBlock)
@@ -169,6 +196,7 @@ TSharedRef<SBox> SInterchangePipelineConfigurationDialog::SpawnPipelineConfigura
 	}
 	else
 	{
+
 		StackTextComboBox = SNew(SHorizontalBox)
 		+ SHorizontalBox::Slot()
 		.VAlign(VAlign_Center)
@@ -207,8 +235,8 @@ TSharedRef<SBox> SInterchangePipelineConfigurationDialog::SpawnPipelineConfigura
 			.AutoHeight()
 			[
 				SNew(SBox)
-				.MinDesiredHeight(35)
-				.MaxDesiredHeight(110)
+				.MinDesiredHeight(50)
+				.MaxDesiredHeight(140)
 				[
 					PipelinesListView.ToSharedRef()
 				]
@@ -370,7 +398,26 @@ void SInterchangePipelineConfigurationDialog::Construct(const FArguments& InArgs
 	//Select the first pipeline
 	if (PipelineListViewItems.Num() > 0)
 	{
-		PipelinesListView->SetSelection(PipelineListViewItems[0], ESelectInfo::Direct);
+		bool bSelectFirst = true;
+		FString LastPipelineName;
+		FString KeyName = CurrentStackName.ToString() + TEXT("_LastSelectedPipeline");
+		if (GConfig->GetString(TEXT("InterchangeSelectPipeline"), *KeyName, LastPipelineName, GEditorPerProjectIni))
+		{
+			for (TObjectPtr<UInterchangePipelineBase> PipelineItem : PipelineListViewItems)
+			{
+				FString PipelineItemName = PipelineItem->GetClass()->GetName();
+				if (PipelineItemName.Equals(LastPipelineName))
+				{
+					PipelinesListView->SetSelection(PipelineItem, ESelectInfo::Direct);
+					bSelectFirst = false;
+					break;
+				}
+			}
+		}
+		if (bSelectFirst)
+		{
+			PipelinesListView->SetSelection(PipelineListViewItems[0], ESelectInfo::Direct);
+		}
 	}
 }
 
@@ -479,7 +526,7 @@ FText SInterchangePipelineConfigurationDialog::GetImportButtonTooltip() const
 	{
 		return InvalidReason.GetValue();
 	}
-	return FText();
+	return LOCTEXT("ImportButtonDefaultTooltip", "Selected pipeline stack will the used for the current import");
 }
 
 void SInterchangePipelineConfigurationDialog::SaveAllPipelineSettings() const
@@ -630,6 +677,10 @@ void SInterchangePipelineConfigurationDialog::OnPipelineSelectionChanged(TObject
 		CurrentSelectedPipeline = InItem;
 	}
 	PipelineConfigurationDetailsView->SetObject(CurrentSelectedPipeline.Get());
+
+	FString CurrentPipelineName = CurrentSelectedPipeline->GetClass()->GetName();
+	FString KeyName = CurrentStackName.ToString() + TEXT("_LastSelectedPipeline");
+	GConfig->SetString(TEXT("InterchangeSelectPipeline"), *KeyName, *CurrentPipelineName, GEditorPerProjectIni);
 }
 
 #undef LOCTEXT_NAMESPACE
diff --git a/Engine/Plugins/Interchange/Editor/Source/Pipelines/Private/SInterchangePipelineConfigurationDialog.h b/Engine/Plugins/Interchange/Editor/Source/Pipelines/Private/SInterchangePipelineConfigurationDialog.h
--- a/Engine/Plugins/Interchange/Editor/Source/Pipelines/Private/SInterchangePipelineConfigurationDialog.h
+++ b/Engine/Plugins/Interchange/Editor/Source/Pipelines/Private/SInterchangePipelineConfigurationDialog.h
@@ -7,6 +7,7 @@
 #include "InterchangePipelineBase.h"
 #include "InterchangePipelineConfigurationBase.h"
 #include "InterchangeSourceData.h"
+#include "Styling/SlateBrush.h"
 #include "Styling/SlateColor.h"
 #include "UObject/GCObject.h"
 #include "UObject/WeakObjectPtrTemplates.h"
@@ -15,6 +16,7 @@
 #include "Widgets/Views/STreeView.h"
 
 struct FPropertyAndParent;
+struct FSlateBrush;
 class IDetailsView;
 class SCheckBox;
 
@@ -26,6 +28,8 @@ public:
 		const TSharedRef<STableViewBase>& OwnerTable,
 		TObjectPtr<UInterchangePipelineBase> InPipelineElement);
 private:
+	const FSlateBrush* GetImageItemIcon() const;
+
 	TObjectPtr<UInterchangePipelineBase> PipelineElement;
 };
 
