Fixed compile errors from merge (just includes)

diff --git a/Engine/Plugins/Interchange/Editor/Source/Pipelines/Private/SInterchangePipelineConfigurationDialog.cpp b/Engine/Plugins/Interchange/Editor/Source/Pipelines/Private/SInterchangePipelineConfigurationDialog.cpp
--- a/Engine/Plugins/Interchange/Editor/Source/Pipelines/Private/SInterchangePipelineConfigurationDialog.cpp
+++ b/Engine/Plugins/Interchange/Editor/Source/Pipelines/Private/SInterchangePipelineConfigurationDialog.cpp
@@ -17,6 +17,7 @@
 #include "SPrimaryButton.h"
 #include "Styling/SlateIconFinder.h"
 #include "Widgets/Input/SCheckBox.h"
+#include "Widgets/Input/STextComboBox.h"
 #include "Widgets/Layout/SUniformGridPanel.h"
 
 #include "Widgets/SWindow.h"
