[NO NON-PII COMMIT MESSAGE]

diff --git a/Engine/Plugins/Interchange/Editor/Source/Pipelines/Private/SInterchangePipelineConfigurationDialog.cpp b/Engine/Plugins/Interchange/Editor/Source/Pipelines/Private/SInterchangePipelineConfigurationDialog.cpp
--- a/Engine/Plugins/Interchange/Editor/Source/Pipelines/Private/SInterchangePipelineConfigurationDialog.cpp
+++ b/Engine/Plugins/Interchange/Editor/Source/Pipelines/Private/SInterchangePipelineConfigurationDialog.cpp
@@ -163,7 +163,7 @@ TSharedRef<SBox> SInterchangePipelineConfigurationDialog::SpawnPipelineConfigura
 	//Only use a combo box if there is more then one stack
 	if (AvailableStacks.Num() > 1)
 	{
-		FText StackComboBoxTooltip = LOCTEXT("StackComboBoxTooltip", "Selected pipeline stack will the used for the current import. To change the pipeline stack used when automating or without dialog please change the default pipeline stacks in the project settings");
+		FText StackComboBoxTooltip = LOCTEXT("StackComboBoxTooltip", "Selected pipeline stack will be used for the current import. To change the pipeline stack used when automating or without dialog please change the default pipeline stacks in the project settings");
 		TextComboBoxPtr = SNew(STextComboBox)
 			.OptionsSource(&AvailableStacks)
 			.OnSelectionChanged(this, &SInterchangePipelineConfigurationDialog::OnStackSelectionChanged)
@@ -518,7 +518,7 @@ FText SInterchangePipelineConfigurationDialog::GetImportButtonTooltip() const
 	{
 		return InvalidReason.GetValue();
 	}
-	return LOCTEXT("ImportButtonDefaultTooltip", "Selected pipeline stack will the used for the current import");
+	return LOCTEXT("ImportButtonDefaultTooltip", "Selected pipeline stack will be used for the current import");
 }
 
 void SInterchangePipelineConfigurationDialog::SaveAllPipelineSettings() const
