Replaced InterchangeResult tab with use of MessageLog window

diff --git a/Engine/Plugins/Interchange/Editor/InterchangeEditor.uplugin b/Engine/Plugins/Interchange/Editor/InterchangeEditor.uplugin
--- a/Engine/Plugins/Interchange/Editor/InterchangeEditor.uplugin
+++ b/Engine/Plugins/Interchange/Editor/InterchangeEditor.uplugin
@@ -38,15 +38,6 @@
 				"Program"
 			]
 		},
-		{
-			"Name": "InterchangeResultsBrowser",
-			"Type": "Editor",
-			"LoadingPhase": "Default",
-			"TargetAllowList": [
-				"Editor",
-				"Program"
-			]
-		},
 		{
 			"Name": "InterchangeEditorUtilities",
 			"Type": "Editor",
diff --git a/Engine/Plugins/Interchange/Editor/Source/InterchangeEditor/InterchangeEditor.Build.cs b/Engine/Plugins/Interchange/Editor/Source/InterchangeEditor/InterchangeEditor.Build.cs
--- a/Engine/Plugins/Interchange/Editor/Source/InterchangeEditor/InterchangeEditor.Build.cs
+++ b/Engine/Plugins/Interchange/Editor/Source/InterchangeEditor/InterchangeEditor.Build.cs
@@ -12,14 +12,17 @@ namespace UnrealBuildTool.Rules
 				new string[]
 				{
                     "Core",
-					"CoreUObject"
+					"CoreUObject",
+					"Engine",
 				}
 			);
 
 			PrivateDependencyModuleNames.AddRange(
 				new string[]
 				{
-					"InterchangeCore"
+					"InterchangeCore",
+					"InterchangeEngine",
+					"MessageLog",
 				}
 			);
 		}
diff --git a/Engine/Plugins/Interchange/Editor/Source/InterchangeEditor/Private/InterchangeEditorModule.cpp b/Engine/Plugins/Interchange/Editor/Source/InterchangeEditor/Private/InterchangeEditorModule.cpp
--- a/Engine/Plugins/Interchange/Editor/Source/InterchangeEditor/Private/InterchangeEditorModule.cpp
+++ b/Engine/Plugins/Interchange/Editor/Source/InterchangeEditor/Private/InterchangeEditorModule.cpp
@@ -3,14 +3,69 @@
 
 #include "InterchangeEditorLog.h"
 
-#include "Modules/ModuleManager.h"
-
+#include "InterchangeManager.h"
 
+#include "Engine/Engine.h"
+#include "IMessageLogListing.h"
+#include "MessageLogModule.h"
+#include "Misc/App.h"
+#include "Misc/CoreDelegates.h"
+#include "Modules/ModuleManager.h"
 
 #define LOCTEXT_NAMESPACE "InterchangeEditorModule"
 
 DEFINE_LOG_CATEGORY(LogInterchangeEditor);
 
+namespace UE::Interchange::InterchangeEditorModule
+{
+	bool HasErrorsOrWarnings(TStrongObjectPtr<UInterchangeResultsContainer> InResultsContainer)
+	{
+		for (UInterchangeResult* Result : InResultsContainer->GetResults())
+		{
+			if (Result->GetResultType() != EInterchangeResultType::Success)
+			{
+				return true;
+			}
+		}
+
+		return false;
+	}
+
+	void LogErrors(TStrongObjectPtr<UInterchangeResultsContainer> InResultsContainer)
+	{
+		// Only showing when we have errors or warnings for now
+		if (FApp::IsUnattended() || !HasErrorsOrWarnings(InResultsContainer))
+		{
+			return;
+		}
+
+		FMessageLogModule& MessageLogModule = FModuleManager::LoadModuleChecked<FMessageLogModule>("MessageLog");
+		TSharedPtr<IMessageLogListing> LogListing = MessageLogModule.GetLogListing(FName("Interchange"));
+
+		if (ensure(LogListing))
+		{
+			const FText LogListingLabel = NSLOCTEXT("InterchangeImport", "Label", "Interchange Import");
+			LogListing->SetLabel(LogListingLabel);
+
+			TArray<TSharedRef<FTokenizedMessage>> TokenizedMessages;
+			bool bHasErrors = false;
+
+			for (UInterchangeResult* Result : InResultsContainer->GetResults())
+			{
+				const EInterchangeResultType ResultType = Result->GetResultType();
+				if (ResultType != EInterchangeResultType::Success)
+				{
+					TokenizedMessages.Add(FTokenizedMessage::Create(ResultType == EInterchangeResultType::Error ? EMessageSeverity::Error : EMessageSeverity::Warning, Result->GetText()));
+					bHasErrors |= ResultType == EInterchangeResultType::Error;
+				}
+			}
+
+			LogListing->AddMessages(TokenizedMessages);
+			LogListing->NotifyIfAnyMessages(NSLOCTEXT("Interchange", "LogAndNotify", "There was some issues with the import."), EMessageSeverity::Info);
+		}
+	}
+}
+
 FInterchangeEditorModule& FInterchangeEditorModule::Get()
 {
 	return FModuleManager::LoadModuleChecked< FInterchangeEditorModule >(INTERCHANGEEDITOR_MODULE_NAME);
@@ -23,9 +78,35 @@ bool FInterchangeEditorModule::IsAvailable()
 
 void FInterchangeEditorModule::StartupModule()
 {
+	using namespace UE::Interchange;
+
+	auto RegisterItems = [this]()
+	{
+		FDelegateHandle InterchangeEditorModuleDelegate;
+
+		UInterchangeManager& InterchangeManager = UInterchangeManager::GetInterchangeManager();
+		InterchangeEditorModuleDelegate = InterchangeManager.OnBatchImportComplete.AddStatic(&InterchangeEditorModule::LogErrors);
+
+		auto UnregisterItems = [InterchangeEditorModuleDelegate]()
+		{
+			UInterchangeManager& InterchangeManager = UInterchangeManager::GetInterchangeManager();
+			InterchangeManager.OnBatchImportComplete.Remove(InterchangeEditorModuleDelegate);
+		};
+
+		InterchangeManager.OnPreDestroyInterchangeManager.AddLambda(UnregisterItems);
+	};
+
+	if (GEngine)
+	{
+		RegisterItems();
+	}
+	else
+	{
+		FCoreDelegates::OnPostEngineInit.AddLambda(RegisterItems);
+	}
 }
 
-IMPLEMENT_MODULE(FInterchangeEditorModule, InterchangeEditor);
+IMPLEMENT_MODULE(FInterchangeEditorModule, InterchangeEditor)
 
 #undef LOCTEXT_NAMESPACE // "InterchangeEditorModule"
 
diff --git a/Engine/Plugins/Interchange/Editor/Source/ResultsBrowser/InterchangeResultsBrowser.Build.cs b/Engine/Plugins/Interchange/Editor/Source/ResultsBrowser/InterchangeResultsBrowser.Build.cs
deleted file mode 100644
diff --git a/Engine/Plugins/Interchange/Editor/Source/ResultsBrowser/Private/InterchangeResultsBrowserModule.cpp b/Engine/Plugins/Interchange/Editor/Source/ResultsBrowser/Private/InterchangeResultsBrowserModule.cpp
deleted file mode 100644
diff --git a/Engine/Plugins/Interchange/Editor/Source/ResultsBrowser/Private/InterchangeResultsBrowserStyle.h b/Engine/Plugins/Interchange/Editor/Source/ResultsBrowser/Private/InterchangeResultsBrowserStyle.h
deleted file mode 100644
diff --git a/Engine/Plugins/Interchange/Editor/Source/ResultsBrowser/Private/SInterchangeResultsBrowserWindow.cpp b/Engine/Plugins/Interchange/Editor/Source/ResultsBrowser/Private/SInterchangeResultsBrowserWindow.cpp
deleted file mode 100644
diff --git a/Engine/Plugins/Interchange/Editor/Source/ResultsBrowser/Private/SInterchangeResultsBrowserWindow.h b/Engine/Plugins/Interchange/Editor/Source/ResultsBrowser/Private/SInterchangeResultsBrowserWindow.h
deleted file mode 100644
diff --git a/Engine/Plugins/Interchange/Editor/Source/ResultsBrowser/Public/InterchangeResultsBrowserModule.h b/Engine/Plugins/Interchange/Editor/Source/ResultsBrowser/Public/InterchangeResultsBrowserModule.h
deleted file mode 100644
