Fixing compilation errors in NoPCH UnrealEditor

diff --git a/Engine/Plugins/Interchange/Editor/Source/Pipelines/Private/InterchangeMaterialXPipelineSettingsCustomization.cpp b/Engine/Plugins/Interchange/Editor/Source/Pipelines/Private/InterchangeMaterialXPipelineSettingsCustomization.cpp
--- a/Engine/Plugins/Interchange/Editor/Source/Pipelines/Private/InterchangeMaterialXPipelineSettingsCustomization.cpp
+++ b/Engine/Plugins/Interchange/Editor/Source/Pipelines/Private/InterchangeMaterialXPipelineSettingsCustomization.cpp
@@ -4,6 +4,7 @@
 #include "IDetailChildrenBuilder.h"
 #include "InterchangeGenericMaterialPipeline.h"
 #include "InterchangeMaterialDefinitions.h"
+#include "Materials/MaterialFunction.h"
 #include "PropertyCustomizationHelpers.h"
 
 TSharedRef<IPropertyTypeCustomization> FInterchangeMaterialXPipelineSettingsCustomization::MakeInstance()
diff --git a/Engine/Plugins/Interchange/Runtime/Source/Pipelines/Private/InterchangeGenericMaterialPipeline.cpp b/Engine/Plugins/Interchange/Runtime/Source/Pipelines/Private/InterchangeGenericMaterialPipeline.cpp
--- a/Engine/Plugins/Interchange/Runtime/Source/Pipelines/Private/InterchangeGenericMaterialPipeline.cpp
+++ b/Engine/Plugins/Interchange/Runtime/Source/Pipelines/Private/InterchangeGenericMaterialPipeline.cpp
@@ -39,6 +39,7 @@
 #include "Materials/MaterialExpressionTransform.h"
 #include "Materials/MaterialExpressionVectorNoise.h"
 #include "Materials/MaterialExpressionVectorParameter.h"
+#include "Materials/MaterialFunction.h"
 #include "Materials/MaterialInstanceConstant.h"
 #include "Materials/MaterialInstanceDynamic.h"
 #include "Misc/CoreMisc.h"
diff --git a/Engine/Plugins/Interchange/Runtime/Source/Pipelines/Public/InterchangeGenericMaterialPipeline.h b/Engine/Plugins/Interchange/Runtime/Source/Pipelines/Public/InterchangeGenericMaterialPipeline.h
--- a/Engine/Plugins/Interchange/Runtime/Source/Pipelines/Public/InterchangeGenericMaterialPipeline.h
+++ b/Engine/Plugins/Interchange/Runtime/Source/Pipelines/Public/InterchangeGenericMaterialPipeline.h
@@ -21,6 +21,7 @@ class UInterchangeMaterialExpressionFactoryNode;
 class UInterchangeMaterialFunctionFactoryNode;
 class UInterchangeMaterialInstanceFactoryNode;
 class UInterchangeResult;
+class UMaterialFunction;
 
 UENUM(BlueprintType)
 enum class EInterchangeMaterialImportOption : uint8
