Fixing heap use after free

diff --git a/Engine/Plugins/Interchange/Editor/Source/Pipelines/Private/InterchangeEditorPipelinesModule.cpp b/Engine/Plugins/Interchange/Editor/Source/Pipelines/Private/InterchangeEditorPipelinesModule.cpp
--- a/Engine/Plugins/Interchange/Editor/Source/Pipelines/Private/InterchangeEditorPipelinesModule.cpp
+++ b/Engine/Plugins/Interchange/Editor/Source/Pipelines/Private/InterchangeEditorPipelinesModule.cpp
@@ -38,6 +38,7 @@ private:
 	TSharedPtr<ISlateStyle> InterchangeEditorPipelineStyle = nullptr;
 
 	TMultiMap<FName, FName> RegisteredPropertySections;
+	TArray<FName> PropertiesTypesToUnregisterOnShutdown;
 };
 
 IMPLEMENT_MODULE(FInterchangeEditorPipelinesModule, InterchangeEditorPipelines)
@@ -48,9 +49,7 @@ void FInterchangeEditorPipelinesModule::StartupModule()
 }
 
 void FInterchangeEditorPipelinesModule::ShutdownModule()
-{
-	ReleaseResources();
-}
+{}
 
 void FInterchangeEditorPipelinesModule::AcquireResources()
 {
@@ -80,7 +79,8 @@ void FInterchangeEditorPipelinesModule::AcquireResources()
 	ClassesToUnregisterOnShutdown.Add(UInterchangePipelineBase::StaticClass()->GetFName());
 	PropertyEditorModule.RegisterCustomClassLayout(ClassesToUnregisterOnShutdown.Last(), FOnGetDetailCustomizationInstance::CreateStatic(&FInterchangePipelineBaseDetailsCustomization::MakeInstance));
 
-	PropertyEditorModule.RegisterCustomPropertyTypeLayout(FMaterialXPipelineSettings::StaticStruct()->GetFName(), FOnGetPropertyTypeCustomizationInstance::CreateStatic(&FInterchangeMaterialXPipelineSettingsCustomization::MakeInstance));
+	PropertiesTypesToUnregisterOnShutdown.Add(FMaterialXPipelineSettings::StaticStruct()->GetFName());
+	PropertyEditorModule.RegisterCustomPropertyTypeLayout(PropertiesTypesToUnregisterOnShutdown.Last(), FOnGetPropertyTypeCustomizationInstance::CreateStatic(&FInterchangeMaterialXPipelineSettingsCustomization::MakeInstance));
 
 	if (!InterchangeEditorPipelineStyle.IsValid())
 	{
@@ -116,8 +116,14 @@ void FInterchangeEditorPipelinesModule::ReleaseResources()
 		{
 			PropertyEditorModule->UnregisterCustomClassLayout(ClassName);
 		}
+
+		for (FName PropertyName : PropertiesTypesToUnregisterOnShutdown)
+		{
+			PropertyEditorModule->UnregisterCustomPropertyTypeLayout(PropertyName);
+		}
 	}
 	ClassesToUnregisterOnShutdown.Empty();
+	PropertiesTypesToUnregisterOnShutdown.Empty();
 
 	UnregisterPropertySectionMappings();
 
@@ -131,8 +137,6 @@ void FInterchangeEditorPipelinesModule::ReleaseResources()
 		AssetTools.UnregisterAssetTypeActions(PipelineBase_TypeActions.ToSharedRef());
 		AssetTools.UnregisterAssetTypeActions(PythonPipelineBase_TypeActions.ToSharedRef());
 	}
-
-	PropertyEditorModule->UnregisterCustomPropertyTypeLayout(FMaterialXPipelineSettings::StaticStruct()->GetFName());
 }
 
 TSharedRef<FPropertySection> FInterchangeEditorPipelinesModule::RegisterPropertySection(FPropertyEditorModule& PropertyModule, FName ClassName, FName SectionName, FText DisplayName)
