Interchange: fix crash in the Interchange import option dialog. When you click to selected no pipeline in the pipeline list.

#rnx

diff --git a/Engine/Plugins/Interchange/Editor/Source/Pipelines/Private/SInterchangePipelineConfigurationDialog.cpp b/Engine/Plugins/Interchange/Editor/Source/Pipelines/Private/SInterchangePipelineConfigurationDialog.cpp
--- a/Engine/Plugins/Interchange/Editor/Source/Pipelines/Private/SInterchangePipelineConfigurationDialog.cpp
+++ b/Engine/Plugins/Interchange/Editor/Source/Pipelines/Private/SInterchangePipelineConfigurationDialog.cpp
@@ -669,10 +669,13 @@ void SInterchangePipelineConfigurationDialog::OnPipelineSelectionChanged(TObject
 		CurrentSelectedPipeline = InItem;
 	}
 	PipelineConfigurationDetailsView->SetObject(CurrentSelectedPipeline.Get());
-
-	FString CurrentPipelineName = CurrentSelectedPipeline->GetClass()->GetName();
-	FString KeyName = CurrentStackName.ToString() + TEXT("_LastSelectedPipeline");
-	GConfig->SetString(TEXT("InterchangeSelectPipeline"), *KeyName, *CurrentPipelineName, GEditorPerProjectIni);
+	
+	if (CurrentSelectedPipeline)
+	{
+		FString CurrentPipelineName = CurrentSelectedPipeline->GetClass()->GetName();
+		FString KeyName = CurrentStackName.ToString() + TEXT("_LastSelectedPipeline");
+		GConfig->SetString(TEXT("InterchangeSelectPipeline"), *KeyName, *CurrentPipelineName, GEditorPerProjectIni);
+	}
 }
 
 #undef LOCTEXT_NAMESPACE
