Interchange: Add LODGroup static mesh pipeline option

#rnx

diff --git a/Engine/Plugins/Interchange/Editor/Source/Pipelines/Private/InterchangeEditorPipelineDetails.cpp b/Engine/Plugins/Interchange/Editor/Source/Pipelines/Private/InterchangeEditorPipelineDetails.cpp
--- a/Engine/Plugins/Interchange/Editor/Source/Pipelines/Private/InterchangeEditorPipelineDetails.cpp
+++ b/Engine/Plugins/Interchange/Editor/Source/Pipelines/Private/InterchangeEditorPipelineDetails.cpp
@@ -10,6 +10,7 @@
 #include "ScopedTransaction.h"
 #include "Styling/StyleColors.h"
 #include "Widgets/Input/SNumericEntryBox.h"
+#include "Widgets/Input/STextComboBox.h"
 #include "Widgets/Layout/SBox.h"
 
 #define LOCTEXT_NAMESPACE "InterchangeEditorPipelineDetails"
@@ -107,6 +108,149 @@ TSharedRef<IDetailCustomization> FInterchangePipelineBaseDetailsCustomization::M
 	return MakeShareable(new FInterchangePipelineBaseDetailsCustomization);
 }
 
+void FInterchangePipelineBaseDetailsCustomization::SetTextComboBoxWidget(IDetailPropertyRow& PropertyRow, const TSharedPtr<IPropertyHandle>& Handle, const TArray<FString>& PossibleValues)
+{
+	if (!Handle.IsValid() || !Handle->IsValidHandle() || PossibleValues.Num() < 1)
+	{
+		//We customize only valid handle and if there is some possible values
+		return;
+	}
+
+	FProperty* PropertyPtr = Handle->GetProperty();
+	if (!PropertyPtr)
+	{
+		return;
+	}
+	const bool bIsNameProperty = PropertyPtr->IsA(FNameProperty::StaticClass());
+	const bool bIsStringProperty = PropertyPtr->IsA(FStrProperty::StaticClass());
+	if (!bIsNameProperty && ! bIsStringProperty)
+	{
+		//We support only FName and FString property
+		return;
+	}
+
+	//Create the data in the Map
+	FInternalComboBoxData& InternalComboBoxData = ComboBoxDataPerProperty.FindOrAdd(Handle);
+	//Generate a combo box row with all the possible value
+	InternalComboBoxData.StringSharedOptions.Reset(PossibleValues.Num());
+	if (bIsNameProperty)
+	{
+		InternalComboBoxData.NameOptions.Reset(PossibleValues.Num());
+		InternalComboBoxData.StringOptions.Empty();
+	}
+	else
+	{
+		InternalComboBoxData.NameOptions.Empty();
+		InternalComboBoxData.StringOptions.Reset(PossibleValues.Num());
+	}
+	
+	for (int32 ValueIndex = 0; ValueIndex < PossibleValues.Num(); ++ValueIndex)
+	{
+		const FString& Value = PossibleValues[ValueIndex];
+		InternalComboBoxData.StringSharedOptions.Add(MakeShareable(new FString(Value)));
+		if (bIsNameProperty)
+		{
+			InternalComboBoxData.NameOptions.Add(FName(*Value));
+		}
+		else
+		{
+			InternalComboBoxData.StringOptions.Add(Value);
+		}
+	}
+
+	TSharedPtr<SWidget> NameWidget;
+	TSharedPtr<SWidget> ValueWidget;
+	FDetailWidgetRow Row;
+	PropertyRow.GetDefaultWidgets(NameWidget, ValueWidget, Row);
+	if (PropertyRow.CustomNameWidget())
+	{
+		//Keep the customize name widget
+		NameWidget = PropertyRow.CustomNameWidget()->Widget;
+	}
+
+	int32 ComboBoxIndex = INDEX_NONE;
+	if (bIsNameProperty)
+	{
+		FName InitialValue;
+		ensure(Handle->GetValue(InitialValue) == FPropertyAccess::Success);
+		ComboBoxIndex = InternalComboBoxData.NameOptions.Find(InitialValue);
+		if (ComboBoxIndex == INDEX_NONE && InternalComboBoxData.NameOptions.Num() > 0)
+		{
+			ComboBoxIndex = 0;
+			ensure(Handle->SetValue(InternalComboBoxData.NameOptions[ComboBoxIndex]) == FPropertyAccess::Success);
+		}
+	}
+	else
+	{
+		FString InitialValue;
+		ensure(Handle->GetValue(InitialValue) == FPropertyAccess::Success);
+		ComboBoxIndex = InternalComboBoxData.StringOptions.Find(InitialValue);
+		if (ComboBoxIndex == INDEX_NONE && InternalComboBoxData.StringOptions.Num() > 0)
+		{
+			ComboBoxIndex = 0;
+			ensure(Handle->SetValue(InternalComboBoxData.StringOptions[ComboBoxIndex]) == FPropertyAccess::Success);
+		}
+	}
+	
+	check(ComboBoxIndex != INDEX_NONE);
+	TWeakPtr<IPropertyHandle> HandlePtr = Handle;
+
+	const bool bShowChildren = true;
+	PropertyRow.CustomWidget(bShowChildren)
+		.NameContent()
+		.MinDesiredWidth(Row.NameWidget.MinWidth)
+		.MaxDesiredWidth(Row.NameWidget.MaxWidth)
+		.HAlign(HAlign_Fill)
+		[
+			NameWidget.ToSharedRef()
+		]
+		.ValueContent()
+		.MinDesiredWidth(Row.ValueWidget.MinWidth)
+		.MaxDesiredWidth(Row.ValueWidget.MaxWidth)
+		.VAlign(VAlign_Center)
+		[
+			SNew(STextComboBox)
+			.Font(IDetailLayoutBuilder::GetDetailFont())
+			.OptionsSource(&InternalComboBoxData.StringSharedOptions)
+			.InitiallySelectedItem(InternalComboBoxData.StringSharedOptions[ComboBoxIndex])
+			.OnSelectionChanged(this, &FInterchangePipelineBaseDetailsCustomization::OnTextComboBoxChanged, HandlePtr)
+		];
+}
+
+void FInterchangePipelineBaseDetailsCustomization::OnTextComboBoxChanged(TSharedPtr<FString> NewValue, ESelectInfo::Type SelectInfo, TWeakPtr<IPropertyHandle> HandlePtr)
+{
+	TSharedPtr<IPropertyHandle> Handle = HandlePtr.Pin();
+	if (Handle.IsValid() && Handle->IsValidHandle())
+	{
+		FProperty* PropertyPtr = Handle->GetProperty();
+		if (!PropertyPtr)
+		{
+			return;
+		}
+		FInternalComboBoxData& InternalComboBoxData = ComboBoxDataPerProperty.FindChecked(Handle);
+		const bool bIsNameProperty = PropertyPtr->IsA(FNameProperty::StaticClass());
+		const bool bIsStringProperty = PropertyPtr->IsA(FStrProperty::StaticClass());
+		if (!bIsNameProperty && !bIsStringProperty)
+		{
+			//We support only FName and FString property
+			return;
+		}
+
+		int32 GroupIndex = InternalComboBoxData.StringSharedOptions.Find(NewValue);
+		if (ensure(GroupIndex != INDEX_NONE))
+		{
+			if (bIsNameProperty)
+			{
+				ensure(Handle->SetValue(InternalComboBoxData.NameOptions[GroupIndex]) == FPropertyAccess::Success);
+			}
+			else
+			{
+				ensure(Handle->SetValue(InternalComboBoxData.StringOptions[GroupIndex]) == FPropertyAccess::Success);
+			}
+		}
+	}
+}
+
 void FInterchangePipelineBaseDetailsCustomization::CustomizeDetails(IDetailLayoutBuilder& DetailBuilder)
 {
 	CachedDetailBuilder = &DetailBuilder;
@@ -203,6 +347,16 @@ void FInterchangePipelineBaseDetailsCustomization::CustomizeDetails(IDetailLayou
 				{
 					LockPropertyHandleRow(PropertyHandle, PropertyRow);
 				}
+
+				//Customize FName and FString properties if the pipeline define some possible value for this property
+				if (PropertyPtr->IsA(FNameProperty::StaticClass()) || PropertyPtr->IsA(FStrProperty::StaticClass()))
+				{
+					TArray<FString> PossibleValues;
+					if (InterchangePipeline->GetPropertyPossibleValues(PropertyPath, PossibleValues))
+					{
+						SetTextComboBoxWidget(PropertyRow, PropertyHandle, PossibleValues);
+					}
+				}
 			}
 			//This is the asset editor mode, we allow users to set default value and locks
 			else
@@ -210,6 +364,7 @@ void FInterchangePipelineBaseDetailsCustomization::CustomizeDetails(IDetailLayou
 				constexpr bool bShowChildren = true;
 				GetGroupPtr();
 				IDetailPropertyRow& PropertyRow = GroupPtr ? GroupPtr->AddPropertyRow(PropertyHandle) : Category.AddProperty(PropertyHandle);
+				
 				TSharedPtr<SWidget> NameWidget;
 				TSharedPtr<SWidget> ValueWidget;
 				FDetailWidgetRow Row;
@@ -387,6 +542,16 @@ void FInterchangePipelineBaseDetailsCustomization::CustomizeDetails(IDetailLayou
 				[
 					ValueWidget.ToSharedRef()
 				];
+
+				//Customize FName and FString properties if the pipeline define some possible value for this property
+				if (PropertyPtr->IsA(FNameProperty::StaticClass()) || PropertyPtr->IsA(FStrProperty::StaticClass()))
+				{
+					TArray<FString> PossibleValues;
+					if (InterchangePipeline->GetPropertyPossibleValues(PropertyPath, PossibleValues))
+					{
+						SetTextComboBoxWidget(PropertyRow, PropertyHandle, PossibleValues);
+					}
+				}
 			}
 		}
 	}
diff --git a/Engine/Plugins/Interchange/Editor/Source/Pipelines/Private/InterchangeEditorPipelineDetails.h b/Engine/Plugins/Interchange/Editor/Source/Pipelines/Private/InterchangeEditorPipelineDetails.h
--- a/Engine/Plugins/Interchange/Editor/Source/Pipelines/Private/InterchangeEditorPipelineDetails.h
+++ b/Engine/Plugins/Interchange/Editor/Source/Pipelines/Private/InterchangeEditorPipelineDetails.h
@@ -4,6 +4,8 @@
 
 #include "IDetailCustomization.h"
 #include "UObject/WeakObjectPtr.h"
+#include "Templates/SharedPointer.h"
+#include "Types/SlateEnums.h"
 
 class SWidget;
 
@@ -43,6 +45,17 @@ private:
 	void LockPropertyHandleRow(const TSharedPtr<IPropertyHandle> PropertyHandle, IDetailPropertyRow& PropertyRow) const;
 	void AddSubCategory(IDetailLayoutBuilder& DetailBuilder, TMap<FName, TMap<FName, TArray<FInternalPropertyData>>>& SubCategoriesPropertiesPerMainCategory);
 	void InternalGetPipelineProperties(const UInterchangePipelineBase* Pipeline, const TArray<FName>& AllCategoryNames, TMap<FName, TArray<FName>>& PropertiesPerCategorys) const;
+	
+	void SetTextComboBoxWidget(IDetailPropertyRow& PropertyRow, const TSharedPtr<IPropertyHandle>& Handle, const TArray<FString>& PossibleValues);
+	void OnTextComboBoxChanged(TSharedPtr<FString> NewValue, ESelectInfo::Type SelectInfo, TWeakPtr<IPropertyHandle> HandlePtr);
+	struct FInternalComboBoxData
+	{
+		TArray<TSharedPtr<FString>> StringSharedOptions;
+		TArray<FString> StringOptions;
+		TArray<FName> NameOptions;
+	};
+	TMap<TSharedPtr<IPropertyHandle>, FInternalComboBoxData> ComboBoxDataPerProperty;
+
 	/** Use MakeInstance to create an instance of this class */
 	FInterchangePipelineBaseDetailsCustomization();
 
diff --git a/Engine/Plugins/Interchange/Runtime/Source/FactoryNodes/Private/InterchangeMeshFactoryNode.cpp b/Engine/Plugins/Interchange/Runtime/Source/FactoryNodes/Private/InterchangeMeshFactoryNode.cpp
--- a/Engine/Plugins/Interchange/Runtime/Source/FactoryNodes/Private/InterchangeMeshFactoryNode.cpp
+++ b/Engine/Plugins/Interchange/Runtime/Source/FactoryNodes/Private/InterchangeMeshFactoryNode.cpp
@@ -182,6 +182,17 @@ bool UInterchangeMeshFactoryNode::RemoveSlotMaterialDependencyUid(const FString&
 	return false;
 }
 
+bool UInterchangeMeshFactoryNode::GetCustomLODGroup(FName& AttributeValue) const
+{
+	IMPLEMENT_NODE_ATTRIBUTE_GETTER(LODGroup, FName)
+}
+
+bool UInterchangeMeshFactoryNode::SetCustomLODGroup(const FName& AttributeValue, bool bAddApplyDelegate)
+{
+	//The LOD Group must be setup directly by the factory with the static mesh API
+	IMPLEMENT_NODE_ATTRIBUTE_SETTER_NODELEGATE(LODGroup, FName)
+}
+
 bool UInterchangeMeshFactoryNode::GetCustomRecomputeNormals(bool& AttributeValue) const
 {
 	IMPLEMENT_NODE_ATTRIBUTE_GETTER(RecomputeNormals, bool)
@@ -348,7 +359,6 @@ void UInterchangeMeshFactoryNode::CopyWithObject(const UInterchangeFactoryBaseNo
 	if (const UInterchangeMeshFactoryNode* MeshFactoryNode = Cast<UInterchangeMeshFactoryNode>(SourceNode))
 	{
 		UClass* Class = GetObjectClass();
-
 		COPY_NODE_DELEGATES_WITH_CUSTOM_DELEGATE(MeshFactoryNode, UInterchangeMeshFactoryNode, RecomputeNormals, bool, Class)
 		COPY_NODE_DELEGATES_WITH_CUSTOM_DELEGATE(MeshFactoryNode, UInterchangeMeshFactoryNode, RecomputeTangents, bool, Class)
 		COPY_NODE_DELEGATES_WITH_CUSTOM_DELEGATE(MeshFactoryNode, UInterchangeMeshFactoryNode, UseMikkTSpace, bool, Class)
diff --git a/Engine/Plugins/Interchange/Runtime/Source/FactoryNodes/Public/InterchangeMeshFactoryNode.h b/Engine/Plugins/Interchange/Runtime/Source/FactoryNodes/Public/InterchangeMeshFactoryNode.h
--- a/Engine/Plugins/Interchange/Runtime/Source/FactoryNodes/Public/InterchangeMeshFactoryNode.h
+++ b/Engine/Plugins/Interchange/Runtime/Source/FactoryNodes/Public/InterchangeMeshFactoryNode.h
@@ -99,6 +99,14 @@ public:
 	UFUNCTION(BlueprintCallable, Category = "Interchange | FactoryNode | Mesh")
 	bool RemoveSlotMaterialDependencyUid(const FString& SlotName);
 
+	/** Query whether normals in the imported mesh are ignored and recomputed. When normals are recomputed the tangents are also recomputed. */
+	UFUNCTION(BlueprintCallable, Category = "Interchange | FactoryNode | Mesh")
+	bool GetCustomLODGroup(FName& AttributeValue) const;
+
+	/** Set whether normals in the imported mesh are ignored and recomputed. When normals are recomputed the tangents are also recomputed. */
+	UFUNCTION(BlueprintCallable, Category = "Interchange | FactoryNode | Mesh")
+	bool SetCustomLODGroup(const FName& AttributeValue, bool bAddApplyDelegate = true);
+
 	/** Query whether normals in the imported mesh are ignored and recomputed. When normals are recomputed the tangents are also recomputed. */
 	UFUNCTION(BlueprintCallable, Category = "Interchange | FactoryNode | Mesh")
 	bool GetCustomRecomputeNormals(bool& AttributeValue) const;
@@ -170,6 +178,7 @@ private:
 	const UE::Interchange::FAttributeKey Macro_CustomVertexColorReplaceKey = UE::Interchange::FAttributeKey(TEXT("VertexColorReplace"));
 	const UE::Interchange::FAttributeKey Macro_CustomVertexColorIgnoreKey = UE::Interchange::FAttributeKey(TEXT("VertexColorIgnore"));
 	const UE::Interchange::FAttributeKey Macro_CustomVertexColorOverrideKey = UE::Interchange::FAttributeKey(TEXT("VertexColorOverride"));
+	const UE::Interchange::FAttributeKey Macro_CustomLODGroupKey = UE::Interchange::FAttributeKey(TEXT("LODGroup"));
 	const UE::Interchange::FAttributeKey Macro_CustomRecomputeNormalsKey = UE::Interchange::FAttributeKey(TEXT("RecomputeNormals"));
 	const UE::Interchange::FAttributeKey Macro_CustomRecomputeTangentsKey = UE::Interchange::FAttributeKey(TEXT("RecomputeTangents"));
 	const UE::Interchange::FAttributeKey Macro_CustomUseMikkTSpaceKey = UE::Interchange::FAttributeKey(TEXT("UseMikkTSpace"));
diff --git a/Engine/Plugins/Interchange/Runtime/Source/Import/Private/Mesh/InterchangeStaticMeshFactory.cpp b/Engine/Plugins/Interchange/Runtime/Source/Import/Private/Mesh/InterchangeStaticMeshFactory.cpp
--- a/Engine/Plugins/Interchange/Runtime/Source/Import/Private/Mesh/InterchangeStaticMeshFactory.cpp
+++ b/Engine/Plugins/Interchange/Runtime/Source/Import/Private/Mesh/InterchangeStaticMeshFactory.cpp
@@ -674,6 +674,17 @@ UInterchangeFactoryBase::FImportAssetResult UInterchangeStaticMeshFactory::EndIm
 	}
 #endif // WITH_EDITORONLY_DATA
 #if WITH_EDITOR
+	//Lod group need to use the static mesh API and cannot use the apply delegate
+	if (!Arguments.ReimportObject)
+	{
+		FName LodGroup = NAME_None;
+		if (StaticMeshFactoryNode->GetCustomLODGroup(LodGroup) && LodGroup != NAME_None)
+		{
+			constexpr bool bRebuildImmediately = false;
+			constexpr bool bAllowModify = false;
+			StaticMesh->SetLODGroup(LodGroup, bRebuildImmediately, bAllowModify);
+		}
+	}
 	FMeshBudgetProjectSettingsUtils::SetLodGroupForStaticMesh(StaticMesh);
 #endif
 
diff --git a/Engine/Plugins/Interchange/Runtime/Source/Pipelines/Private/InterchangeGenericAssetsPipeline.cpp b/Engine/Plugins/Interchange/Runtime/Source/Pipelines/Private/InterchangeGenericAssetsPipeline.cpp
--- a/Engine/Plugins/Interchange/Runtime/Source/Pipelines/Private/InterchangeGenericAssetsPipeline.cpp
+++ b/Engine/Plugins/Interchange/Runtime/Source/Pipelines/Private/InterchangeGenericAssetsPipeline.cpp
@@ -119,6 +119,39 @@ void UInterchangeGenericAssetsPipeline::AdjustSettingsForContext(EInterchangePip
 	}
 }
 
+#if WITH_EDITOR
+bool UInterchangeGenericAssetsPipeline::GetPropertyPossibleValues(const FName PropertyPath, TArray<FString>& PossibleValues)
+{
+	const FString PropertyPathString = PropertyPath.ToString();
+	if (MaterialPipeline && PropertyPathString.StartsWith(UInterchangeGenericMaterialPipeline::StaticClass()->GetPathName()))
+	{
+		if (MaterialPipeline->GetPropertyPossibleValues(PropertyPath, PossibleValues))
+		{
+			return true;
+		}
+	}
+	
+	if (MeshPipeline && PropertyPathString.StartsWith(UInterchangeGenericMeshPipeline::StaticClass()->GetPathName()))
+	{
+		if (MeshPipeline->GetPropertyPossibleValues(PropertyPath, PossibleValues))
+		{
+			return true;
+		}
+	}
+	
+	if (AnimationPipeline && PropertyPathString.StartsWith(UInterchangeGenericAnimationPipeline::StaticClass()->GetPathName()))
+	{
+		if (AnimationPipeline->GetPropertyPossibleValues(PropertyPath, PossibleValues))
+		{
+			return true;
+		}
+	}
+
+	//If we did not find any property call the super implementation
+	return Super::GetPropertyPossibleValues(PropertyPath, PossibleValues);
+}
+#endif
+
 void UInterchangeGenericAssetsPipeline::ExecutePipeline(UInterchangeBaseNodeContainer* InBaseNodeContainer, const TArray<UInterchangeSourceData*>& InSourceDatas)
 {
 	check(CommonSkeletalMeshesAndAnimationsProperties);
diff --git a/Engine/Plugins/Interchange/Runtime/Source/Pipelines/Private/InterchangeGenericMeshPipeline.cpp b/Engine/Plugins/Interchange/Runtime/Source/Pipelines/Private/InterchangeGenericMeshPipeline.cpp
--- a/Engine/Plugins/Interchange/Runtime/Source/Pipelines/Private/InterchangeGenericMeshPipeline.cpp
+++ b/Engine/Plugins/Interchange/Runtime/Source/Pipelines/Private/InterchangeGenericMeshPipeline.cpp
@@ -104,6 +104,30 @@ void UInterchangeGenericMeshPipeline::PreDialogCleanup(const FName PipelineStack
 	PhysicsAsset = nullptr;
 }
 
+#if WITH_EDITOR
+bool UInterchangeGenericMeshPipeline::GetPropertyPossibleValues(const FName PropertyPath, TArray<FString>& PossibleValues)
+{
+	FString PropertyPathString = PropertyPath.ToString();
+	int32 PropertyNameIndex = INDEX_NONE;
+	if (PropertyPathString.FindLastChar(':', PropertyNameIndex))
+	{
+		PropertyPathString = PropertyPathString.RightChop(PropertyNameIndex+1);
+	}
+	if (PropertyPathString.Equals(GET_MEMBER_NAME_STRING_CHECKED(UInterchangeGenericMeshPipeline, LodGroup)))
+	{
+		TArray<FName> LODGroupNames;
+		UStaticMesh::GetLODGroups(LODGroupNames);
+		for (int32 GroupIndex = 0; GroupIndex < LODGroupNames.Num(); ++GroupIndex)
+		{
+			PossibleValues.Add(LODGroupNames[GroupIndex].GetPlainNameString());
+		}
+		return true;
+	}
+	//If we did not find any property call the super implementation
+	return Super::GetPropertyPossibleValues(PropertyPath, PossibleValues);
+}
+#endif
+
 void UInterchangeGenericMeshPipeline::ExecutePipeline(UInterchangeBaseNodeContainer* InBaseNodeContainer, const TArray<UInterchangeSourceData*>& InSourceDatas)
 {
 	if (!InBaseNodeContainer)
diff --git a/Engine/Plugins/Interchange/Runtime/Source/Pipelines/Private/InterchangeGenericStaticMeshPipeline.cpp b/Engine/Plugins/Interchange/Runtime/Source/Pipelines/Private/InterchangeGenericStaticMeshPipeline.cpp
--- a/Engine/Plugins/Interchange/Runtime/Source/Pipelines/Private/InterchangeGenericStaticMeshPipeline.cpp
+++ b/Engine/Plugins/Interchange/Runtime/Source/Pipelines/Private/InterchangeGenericStaticMeshPipeline.cpp
@@ -430,6 +430,8 @@ UInterchangeStaticMeshFactoryNode* UInterchangeGenericMeshPipeline::CreateStatic
 		break;
 	}
 
+	StaticMeshFactoryNode->SetCustomLODGroup(LodGroup);
+
 	//Common meshes build options
 	StaticMeshFactoryNode->SetCustomRecomputeNormals(CommonMeshesProperties->bRecomputeNormals);
 	StaticMeshFactoryNode->SetCustomRecomputeTangents(CommonMeshesProperties->bRecomputeTangents);
diff --git a/Engine/Plugins/Interchange/Runtime/Source/Pipelines/Public/InterchangeGenericAssetsPipeline.h b/Engine/Plugins/Interchange/Runtime/Source/Pipelines/Public/InterchangeGenericAssetsPipeline.h
--- a/Engine/Plugins/Interchange/Runtime/Source/Pipelines/Public/InterchangeGenericAssetsPipeline.h
+++ b/Engine/Plugins/Interchange/Runtime/Source/Pipelines/Public/InterchangeGenericAssetsPipeline.h
@@ -82,6 +82,11 @@ public:
 
 
 	virtual void AdjustSettingsForContext(EInterchangePipelineContext ImportType, TObjectPtr<UObject> ReimportAsset) override;
+
+#if WITH_EDITOR
+	virtual bool GetPropertyPossibleValues(const FName PropertyPath, TArray<FString>& PossibleValues) override;
+#endif
+
 protected:
 
 	virtual void ExecutePipeline(UInterchangeBaseNodeContainer* InBaseNodeContainer, const TArray<UInterchangeSourceData*>& InSourceDatas) override;
diff --git a/Engine/Plugins/Interchange/Runtime/Source/Pipelines/Public/InterchangeGenericMeshPipeline.h b/Engine/Plugins/Interchange/Runtime/Source/Pipelines/Public/InterchangeGenericMeshPipeline.h
--- a/Engine/Plugins/Interchange/Runtime/Source/Pipelines/Public/InterchangeGenericMeshPipeline.h
+++ b/Engine/Plugins/Interchange/Runtime/Source/Pipelines/Public/InterchangeGenericMeshPipeline.h
@@ -48,6 +48,10 @@ public:
 	UPROPERTY(EditAnywhere, BlueprintReadWrite, Category = "Static Meshes")
 	bool bCombineStaticMeshes = false;
 
+	/** The LODGroup to associate with this mesh when it is imported */
+	UPROPERTY(EditAnywhere, BlueprintReadWrite, Category = "Static Meshes")
+	FName LodGroup = NAME_None;
+
 	/**
 	 * If enabled, meshes with certain prefixes will be imported as collision primitives for the mesh with the corresponding unprefixed name.
 	 * 
@@ -187,6 +191,10 @@ public:
 
 	virtual void PreDialogCleanup(const FName PipelineStackName) override;
 
+#if WITH_EDITOR
+	virtual bool GetPropertyPossibleValues(const FName PropertyPath, TArray<FString>& PossibleValues) override;
+#endif
+
 protected:
 	virtual void ExecutePipeline(UInterchangeBaseNodeContainer* InBaseNodeContainer, const TArray<UInterchangeSourceData*>& InSourceDatas) override;
 
diff --git a/Engine/Source/Runtime/Interchange/Core/Public/InterchangePipelineBase.h b/Engine/Source/Runtime/Interchange/Core/Public/InterchangePipelineBase.h
--- a/Engine/Source/Runtime/Interchange/Core/Public/InterchangePipelineBase.h
+++ b/Engine/Source/Runtime/Interchange/Core/Public/InterchangePipelineBase.h
@@ -300,6 +300,18 @@ public:
 	bool CanEditPropertiesStates() { return bAllowPropertyStatesEdition; }
 	bool IsReimportContext() { return bIsReimportContext; }
 
+#if WITH_EDITOR
+	/*
+	 * FName or FString properties can have a dynamic set of possible values. When displaying 
+	 * a pipeline property in the property editor, this function will be call for all FString and FName
+	 * properties and will generate a combo box with the PossibleValues if the function return true.
+	 */
+	virtual bool GetPropertyPossibleValues(const FName PropertyPath, TArray<FString>& PossibleValues)
+	{
+		return false;
+	}
+#endif //WITH_EDITOR
+
 protected:
 
 	UE_DEPRECATED(5.2, "This function is replace by ExecutePipeline.")
