[Interchange] GLTF pipeline's Build Health Text conflict warning fix.

diff --git a/Engine/Plugins/Interchange/Editor/Source/Pipelines/Private/InterchangeGLTFPipelineCustomizations.cpp b/Engine/Plugins/Interchange/Editor/Source/Pipelines/Private/InterchangeGLTFPipelineCustomizations.cpp
--- a/Engine/Plugins/Interchange/Editor/Source/Pipelines/Private/InterchangeGLTFPipelineCustomizations.cpp
+++ b/Engine/Plugins/Interchange/Editor/Source/Pipelines/Private/InterchangeGLTFPipelineCustomizations.cpp
@@ -91,7 +91,7 @@ void FInterchangeGLTFPipelineSettingsCustomization::CustomizeDetails(IDetailLayo
 
 	IDetailCategoryBuilder& GLTFPredefinedMaterialLibraryCategory = DetailBuilder.EditCategory("GLTFPredefinedMaterialLibrary");
 
-	GLTFPredefinedMaterialLibraryCategory.SetDisplayName(NSLOCTEXT("InterchangeGLTFPipelineCustomization", "CategoryDisplayName", "GLTF Predefined Material Library"));
+	GLTFPredefinedMaterialLibraryCategory.SetDisplayName(NSLOCTEXT("InterchangeGLTFPipelineSettingsCustomization", "CategoryDisplayName", "GLTF Predefined Material Library"));
 
 	uint32 NumChildren = 0;
 	PairingsHandle->GetNumChildren(NumChildren);
