[Interchange] GLTF Pipeline customization allowed class fix

diff --git a/Engine/Plugins/Interchange/Editor/Source/Pipelines/Private/InterchangeGLTFPipelineCustomizations.cpp b/Engine/Plugins/Interchange/Editor/Source/Pipelines/Private/InterchangeGLTFPipelineCustomizations.cpp
--- a/Engine/Plugins/Interchange/Editor/Source/Pipelines/Private/InterchangeGLTFPipelineCustomizations.cpp
+++ b/Engine/Plugins/Interchange/Editor/Source/Pipelines/Private/InterchangeGLTFPipelineCustomizations.cpp
@@ -2,21 +2,12 @@
 
 #include "InterchangeGLTFPipelineCustomizations.h"
 
-//#include "InterchangeMaterialDefinitions.h"
 #include "InterchangeglTFPipeline.h"
-#include "Materials/MaterialFunction.h"
+#include "Materials/MaterialInstance.h"
 
 #include "DetailLayoutBuilder.h"
 #include "Modules/ModuleManager.h"
 #include "PropertyCustomizationHelpers.h"
-//#include "DetailWidgetRow.h"
-//#include "IDetailGroup.h"
-//#include "InterchangePipelineBase.h"
-//#include "Nodes/InterchangeBaseNode.h"
-//#include "ScopedTransaction.h"
-//#include "Styling/StyleColors.h"
-//#include "Widgets/Input/SNumericEntryBox.h"
-//#include "Widgets/Layout/SBox.h"
 
 TSharedRef<IDetailCustomization> FInterchangeGLTFPipelineCustomization::MakeInstance()
 {
@@ -120,7 +111,7 @@ void FInterchangeGLTFPipelineSettingsCustomization::CustomizeDetails(IDetailLayo
 			.HAlign(HAlign_Fill)
 			[
 				SNew(SObjectPropertyEntryBox)
-				.AllowedClass(UMaterialFunction::StaticClass())
+				.AllowedClass(UMaterialInstance::StaticClass())
 				.PropertyHandle(ChildPropertyHandle)
 			]
 		];
