[Interchange] GLTF pipeline UI fix

diff --git a/Engine/Plugins/Interchange/Editor/Source/Pipelines/Private/InterchangeGLTFPipelineCustomizations.cpp b/Engine/Plugins/Interchange/Editor/Source/Pipelines/Private/InterchangeGLTFPipelineCustomizations.cpp
--- a/Engine/Plugins/Interchange/Editor/Source/Pipelines/Private/InterchangeGLTFPipelineCustomizations.cpp
+++ b/Engine/Plugins/Interchange/Editor/Source/Pipelines/Private/InterchangeGLTFPipelineCustomizations.cpp
@@ -42,7 +42,7 @@ void FInterchangeGLTFPipelineCustomization::CustomizeDetails(IDetailLayoutBuilde
 
 	if (Pipeline->GLTFPipelineSettings)
 	{
-		Pipeline->GLTFPipelineSettings->SetMaterialParentsEditible(Pipeline->CanEditPropertiesStates());
+		Pipeline->GLTFPipelineSettings->SetMaterialParentsEditible(false);
 	}
 
 	DetailsView->SetObject(Pipeline->GLTFPipelineSettings);
@@ -111,6 +111,9 @@ void FInterchangeGLTFPipelineSettingsCustomization::CustomizeDetails(IDetailLayo
 		];
 	}
 
+	//Reset Editibility for Project Settings -> Interchange GLTF
+	GLTFPipelineSettings->SetMaterialParentsEditible(true);
+
 	IDetailCategoryBuilder& GLTFPredefinedMaterialLibraryCategory = DetailBuilder.EditCategory("PredefinedglTFMaterialLibrary", FText::FromString("Predefined glTF Material Library"));
 
 	DetailBuilder.HideCategory(GET_MEMBER_NAME_CHECKED(UGLTFPipelineSettings, MaterialParents));
