[Interchange] GLTF import options - Remove Gltf Pipeline settings UI.

diff --git a/Engine/Plugins/Interchange/Editor/Source/Pipelines/Private/InterchangeGLTFPipelineCustomizations.cpp b/Engine/Plugins/Interchange/Editor/Source/Pipelines/Private/InterchangeGLTFPipelineCustomizations.cpp
--- a/Engine/Plugins/Interchange/Editor/Source/Pipelines/Private/InterchangeGLTFPipelineCustomizations.cpp
+++ b/Engine/Plugins/Interchange/Editor/Source/Pipelines/Private/InterchangeGLTFPipelineCustomizations.cpp
@@ -8,6 +8,7 @@
 #include "DetailLayoutBuilder.h"
 #include "Modules/ModuleManager.h"
 #include "PropertyCustomizationHelpers.h"
+#include "Widgets/Text/STextBlock.h"
 
 TSharedRef<IDetailCustomization> FInterchangeGLTFPipelineCustomization::MakeInstance()
 {
@@ -45,16 +46,32 @@ void FInterchangeGLTFPipelineCustomization::CustomizeDetails(IDetailLayoutBuilde
 		Pipeline->GLTFPipelineSettings->SetMaterialParentsEditible(false);
 	}
 
-	DetailsView->SetObject(Pipeline->GLTFPipelineSettings);
+	if (Pipeline->CanEditPropertiesStates())
+	{
+		DetailsView->SetObject(Pipeline->GLTFPipelineSettings);
 
-	GLTFCategory.AddCustomRow(NSLOCTEXT("InterchangeGLTFPipelineCustomization", "GLTFPredefinedMaterialLibrary", "Predefined Material Functions"))
-	[
-		SNew( SVerticalBox )
-		+ SVerticalBox::Slot()
-		[
-			DetailsView.ToSharedRef()
-		]
-	];
+		GLTFCategory.AddCustomRow(NSLOCTEXT("InterchangeGLTFPipelineCustomization", "GLTFPredefinedMaterialLibrary", "Predefined Material Functions"))
+			[
+				SNew(SVerticalBox)
+				+ SVerticalBox::Slot()
+			[
+				DetailsView.ToSharedRef()
+			]
+			];
+	}
+	else
+	{
+		GLTFCategory.AddCustomRow(NSLOCTEXT("InterchangeGLTFPipelineMaterialSubstitution::Message::Row", "GLTFPredefinedMaterialSubstitutionMessageRow", "GLTF Pipeline Material Substitution Message"))
+			.NameContent()
+			.VAlign(VAlign_Center)
+			[
+				SNew(STextBlock)
+				.Font(IDetailLayoutBuilder::GetDetailFont())
+				.Text(NSLOCTEXT("InterchangeGLTFPipelineCustomization::Message", "GLTFPredefinedMaterialSubstitutionMessage", "Substitution can be custumized in 'Project Settings > Interchange GLTF'."))
+				.AutoWrapText(true)
+			];
+	}
+	
 }
 
 void FInterchangeGLTFPipelineSettingsCustomization::CustomizeDetails(IDetailLayoutBuilder& DetailBuilder)
diff --git a/Engine/Plugins/Interchange/Runtime/Source/Pipelines/Public/InterchangeglTFPipeline.h b/Engine/Plugins/Interchange/Runtime/Source/Pipelines/Public/InterchangeglTFPipeline.h
--- a/Engine/Plugins/Interchange/Runtime/Source/Pipelines/Public/InterchangeglTFPipeline.h
+++ b/Engine/Plugins/Interchange/Runtime/Source/Pipelines/Public/InterchangeglTFPipeline.h
@@ -51,7 +51,6 @@ class INTERCHANGEPIPELINES_API UInterchangeGLTFPipeline : public UInterchangePip
 	UInterchangeGLTFPipeline();
 
 public:
-	UPROPERTY(EditAnywhere, BlueprintReadWrite, Category = "glTF", meta = (DisplayName = "glTF Pipeline Settings"))
 	TObjectPtr<UGLTFPipelineSettings> GLTFPipelineSettings;
 
 	UPROPERTY(EditAnywhere, BlueprintReadWrite, Category = "glTF", meta = (DisplayName = "Use glTF Material Instance Library"))
