Fix typo in GLTF Pipeline

diff --git a/Engine/Plugins/Interchange/Editor/Source/Pipelines/Private/InterchangeGLTFPipelineCustomizations.cpp b/Engine/Plugins/Interchange/Editor/Source/Pipelines/Private/InterchangeGLTFPipelineCustomizations.cpp
--- a/Engine/Plugins/Interchange/Editor/Source/Pipelines/Private/InterchangeGLTFPipelineCustomizations.cpp
+++ b/Engine/Plugins/Interchange/Editor/Source/Pipelines/Private/InterchangeGLTFPipelineCustomizations.cpp
@@ -67,7 +67,7 @@ void FInterchangeGLTFPipelineCustomization::CustomizeDetails(IDetailLayoutBuilde
 			[
 				SNew(STextBlock)
 				.Font(IDetailLayoutBuilder::GetDetailFont())
-				.Text(NSLOCTEXT("InterchangeGLTFPipelineCustomization::Message", "GLTFPredefinedMaterialSubstitutionMessage", "Substitution can be custumized in 'Project Settings > Interchange GLTF'."))
+				.Text(NSLOCTEXT("InterchangeGLTFPipelineCustomization::Message", "GLTFPredefinedMaterialSubstitutionMessage", "Substitution can be customized in 'Project Settings > Interchange GLTF'."))
 				.AutoWrapText(true)
 			];
 	}
