Interchange: fix reset to default so the selection is set properly on the new duplicated pipeline.

#rnx

diff --git a/Engine/Plugins/Interchange/Editor/Source/Pipelines/Private/SInterchangePipelineConfigurationDialog.cpp b/Engine/Plugins/Interchange/Editor/Source/Pipelines/Private/SInterchangePipelineConfigurationDialog.cpp
--- a/Engine/Plugins/Interchange/Editor/Source/Pipelines/Private/SInterchangePipelineConfigurationDialog.cpp
+++ b/Engine/Plugins/Interchange/Editor/Source/Pipelines/Private/SInterchangePipelineConfigurationDialog.cpp
@@ -480,6 +480,8 @@ FReply SInterchangePipelineConfigurationDialog::OnResetToDefault()
 									PipelineListViewItems.RemoveAt(PipelineIndex, 1, false);
 									PipelineListViewItems.Insert(GeneratedPipeline, PipelineIndex);
 									PipelineConfigurationDetailsView->SetObject(GeneratedPipeline, true);
+									PipelinesListView->SetSelection(PipelineListViewItems[PipelineIndex], ESelectInfo::Direct);
+									PipelinesListView->RequestListRefresh();
 									break;
 								}
 							}
