Unshelved from pending changelist '25965791':

[interchange] Add the name of the pipeline assets in the import window
- The pipeline names in the Import Asset dialog using interchange pipelines have actual asset names followed by the class name in paranthesis.

diff --git a/Engine/Plugins/Interchange/Editor/Source/Pipelines/Private/SInterchangePipelineConfigurationDialog.cpp b/Engine/Plugins/Interchange/Editor/Source/Pipelines/Private/SInterchangePipelineConfigurationDialog.cpp
--- a/Engine/Plugins/Interchange/Editor/Source/Pipelines/Private/SInterchangePipelineConfigurationDialog.cpp
+++ b/Engine/Plugins/Interchange/Editor/Source/Pipelines/Private/SInterchangePipelineConfigurationDialog.cpp
@@ -29,14 +29,20 @@ const FName ReimportStackName = TEXT("ReimportPipeline");
 void SInterchangePipelineItem::Construct(
 	const FArguments& InArgs,
 	const TSharedRef<STableViewBase>& OwnerTable,
-	TObjectPtr<UInterchangePipelineBase> InPipelineElement)
+	TSharedPtr<FInterchangePipelineItemType> InPipelineElement)
 {
 	PipelineElement = InPipelineElement;
-	TObjectPtr<UInterchangePipelineBase> PipelineElementPtr = PipelineElement.Get();
+	TObjectPtr<UInterchangePipelineBase> PipelineElementPtr = PipelineElement->Pipeline;
 	check(PipelineElementPtr.Get());
-	FText PipelineName = PipelineElementPtr.Get() ? FText::FromString(PipelineElementPtr->GetClass()->GetName()) : LOCTEXT("InvalidPipelineName", "Invalid Pipeline");
-	STableRow<TObjectPtr<UInterchangePipelineBase>>::Construct(
-		STableRow<TObjectPtr<UInterchangePipelineBase>>::FArguments()
+	FText PipelineName = LOCTEXT("InvalidPipelineName", "Invalid Pipeline");
+	if (PipelineElementPtr.Get())
+	{
+		FString PipelineNameString = FString::Printf(TEXT("%s (%s)"), *PipelineElement->DisplayName, *PipelineElementPtr->GetClass()->GetName());
+		PipelineName = FText::FromString(PipelineNameString);
+	}
+		
+	STableRow<TSharedPtr<FInterchangePipelineItemType>>::Construct(
+		STableRow<TSharedPtr<FInterchangePipelineItemType>>::FArguments()
 		.Content()
 		[
 			SNew(SHorizontalBox)
@@ -125,10 +131,10 @@ TSharedRef<SBox> SInterchangePipelineConfigurationDialog::SpawnPipelineConfigura
 		TSharedPtr<FString> StackNamePtr = MakeShared<FString>(Stack.StackName.ToString());
 		if (CurrentStackName == Stack.StackName)
 		{
-			for (UInterchangePipelineBase* DefaultPipeline : Stack.Pipelines)
+			for (const TObjectPtr<UInterchangePipelineBase>& DefaultPipeline : Stack.Pipelines)
 			{
 				check(DefaultPipeline);
-				if (UInterchangePipelineBase* GeneratedPipeline = UE::Interchange::GeneratePipelineInstance(DefaultPipeline))
+				if (UInterchangePipelineBase* GeneratedPipeline = UE::Interchange::GeneratePipelineInstanceInSourceAssetPackage(DefaultPipeline))
 				{
 					GeneratedPipeline->TransferAdjustSettings(DefaultPipeline);
 					if (Stack.StackName == ReimportStackName)
@@ -143,7 +149,8 @@ TSharedRef<SBox> SInterchangePipelineConfigurationDialog::SpawnPipelineConfigura
 						GeneratedPipeline->LoadSettings(Stack.StackName);
 						GeneratedPipeline->PreDialogCleanup(Stack.StackName);
 					}
-					PipelineListViewItems.Add(GeneratedPipeline);
+
+					PipelineListViewItems.Add(MakeShareable(new FInterchangePipelineItemType{ DefaultPipeline.GetName(), GeneratedPipeline}));
 				}
 			}
 			SelectedStack = StackNamePtr;
@@ -258,7 +265,6 @@ TSharedRef<SBox> SInterchangePipelineConfigurationDialog::SpawnPipelineConfigura
 	return PipelineConfigurationPanelBox;
 }
 
-
 void SInterchangePipelineConfigurationDialog::Construct(const FArguments& InArgs)
 {
 	//Make sure there is a valid default value
@@ -395,9 +401,9 @@ void SInterchangePipelineConfigurationDialog::Construct(const FArguments& InArgs
 		FString KeyName = CurrentStackName.ToString() + TEXT("_LastSelectedPipeline");
 		if (GConfig->GetString(TEXT("InterchangeSelectPipeline"), *KeyName, LastPipelineName, GEditorPerProjectIni))
 		{
-			for (TObjectPtr<UInterchangePipelineBase> PipelineItem : PipelineListViewItems)
+			for (TSharedPtr<FInterchangePipelineItemType> PipelineItem : PipelineListViewItems)
 			{
-				FString PipelineItemName = PipelineItem->GetClass()->GetName();
+				FString PipelineItemName = PipelineItem->Pipeline->GetClass()->GetName();
 				if (PipelineItemName.Equals(LastPipelineName))
 				{
 					PipelinesListView->SetSelection(PipelineItem, ESelectInfo::Direct);
@@ -463,22 +469,24 @@ FReply SInterchangePipelineConfigurationDialog::OnResetToDefault()
 				{
 					continue;
 				}
-				for (UInterchangePipelineBase* DefaultPipeline : Stack.Pipelines)
+				for (const TObjectPtr<UInterchangePipelineBase>& DefaultPipeline : Stack.Pipelines)
 				{
 					//We assume the pipelines inside one stack are all different classes, we use the class to know which default asset we need to duplicate
 					if (DefaultPipeline->GetClass() == PipelineClass)
 					{
+						
 						for(int32 PipelineIndex = 0; PipelineIndex < PipelineListViewItems.Num(); ++PipelineIndex)
 						{
-							TObjectPtr<UInterchangePipelineBase> PipelineElement = PipelineListViewItems[PipelineIndex];
+							
+							TObjectPtr<UInterchangePipelineBase> PipelineElement = PipelineListViewItems[PipelineIndex]->Pipeline;
+
 							if (PipelineElement.Get() == Pipeline)
 							{
-								if (UInterchangePipelineBase* GeneratedPipeline = UE::Interchange::GeneratePipelineInstance(DefaultPipeline))
+								if (UInterchangePipelineBase* GeneratedPipeline = UE::Interchange::GeneratePipelineInstanceInSourceAssetPackage(DefaultPipeline))
 								{
 									GeneratedPipeline->TransferAdjustSettings(DefaultPipeline);
 									//Switch the pipeline the element point on
-									PipelineListViewItems.RemoveAt(PipelineIndex, 1, false);
-									PipelineListViewItems.Insert(GeneratedPipeline, PipelineIndex);
+									PipelineListViewItems[PipelineIndex]->Pipeline = GeneratedPipeline;
 									PipelineConfigurationDetailsView->SetObject(GeneratedPipeline, true);
 									PipelinesListView->SetSelection(PipelineListViewItems[PipelineIndex], ESelectInfo::Direct);
 									PipelinesListView->RequestListRefresh();
@@ -496,10 +504,10 @@ FReply SInterchangePipelineConfigurationDialog::OnResetToDefault()
 
 bool SInterchangePipelineConfigurationDialog::ValidateAllPipelineSettings(TOptional<FText>& OutInvalidReason) const
 {
-	for (TObjectPtr<UInterchangePipelineBase> PipelineElement : PipelineListViewItems)
+	for (TSharedPtr<FInterchangePipelineItemType> PipelineElement : PipelineListViewItems)
 	{
-		check(PipelineElement);
-		if (!PipelineElement->IsSettingsAreValid(OutInvalidReason))
+		check(PipelineElement->Pipeline);
+		if (!PipelineElement->Pipeline->IsSettingsAreValid(OutInvalidReason))
 		{
 			return false;
 		}
@@ -525,11 +533,11 @@ FText SInterchangePipelineConfigurationDialog::GetImportButtonTooltip() const
 
 void SInterchangePipelineConfigurationDialog::SaveAllPipelineSettings() const
 {
-	for (const TObjectPtr<UInterchangePipelineBase> PipelineElement : PipelineListViewItems)
+	for (TSharedPtr<FInterchangePipelineItemType> PipelineElement : PipelineListViewItems)
 	{
-		if (PipelineElement)
+		if (PipelineElement->Pipeline)
 		{
-			PipelineElement->SaveSettings(CurrentStackName);
+			PipelineElement->Pipeline->SaveSettings(CurrentStackName);
 		}
 	}
 }
@@ -547,13 +555,13 @@ void SInterchangePipelineConfigurationDialog::ClosePipelineConfiguration(const E
 		bImportAll = UseSameSettingsForAllCheckBox->IsChecked();
 		
 		//Fill the OutPipelines array
-		for (const TObjectPtr<UInterchangePipelineBase> PipelineElement : PipelineListViewItems)
+		for (TSharedPtr<FInterchangePipelineItemType> PipelineElement : PipelineListViewItems)
 		{
-			OutPipelines->Add(PipelineElement.Get());
+			OutPipelines->Add(PipelineElement->Pipeline);
 		}
 	}
 
-	//Save the settings only if its not a reimport
+	//Save the settings only if its not a re-import
 	if (!bReimport)
 	{
 		SaveAllPipelineSettings();
@@ -629,10 +637,10 @@ void SInterchangePipelineConfigurationDialog::OnStackSelectionChanged(TSharedPtr
 		{
 			continue;
 		}
-		for (UInterchangePipelineBase* DefaultPipeline : Stack.Pipelines)
+		for (const TObjectPtr<UInterchangePipelineBase>& DefaultPipeline : Stack.Pipelines)
 		{
 			check(DefaultPipeline);
-			if (UInterchangePipelineBase* GeneratedPipeline = UE::Interchange::GeneratePipelineInstance(DefaultPipeline))
+			if (UInterchangePipelineBase* GeneratedPipeline = UE::Interchange::GeneratePipelineInstanceInSourceAssetPackage(DefaultPipeline))
 			{
 				GeneratedPipeline->TransferAdjustSettings(DefaultPipeline);
 				if (Stack.StackName != ReimportStackName)
@@ -641,7 +649,7 @@ void SInterchangePipelineConfigurationDialog::OnStackSelectionChanged(TSharedPtr
 					GeneratedPipeline->LoadSettings(Stack.StackName);
 					GeneratedPipeline->PreDialogCleanup(Stack.StackName);
 				}
-				PipelineListViewItems.Add(GeneratedPipeline);
+				PipelineListViewItems.Add(MakeShareable(new FInterchangePipelineItemType{ DefaultPipeline->GetName(), GeneratedPipeline}));
 			}
 		}
 	}
@@ -654,21 +662,20 @@ void SInterchangePipelineConfigurationDialog::OnStackSelectionChanged(TSharedPtr
 	PipelinesListView->RequestListRefresh();
 }
 
-
 TSharedRef<ITableRow> SInterchangePipelineConfigurationDialog::MakePipelineListRowWidget(
-	TObjectPtr<UInterchangePipelineBase> InElement,
+	TSharedPtr<FInterchangePipelineItemType> InElement,
 	const TSharedRef<STableViewBase>& OwnerTable)
 {
-	check(InElement);
+	check(InElement->Pipeline);
 	return SNew(SInterchangePipelineItem, OwnerTable, InElement);
 }
 
-void SInterchangePipelineConfigurationDialog::OnPipelineSelectionChanged(TObjectPtr<UInterchangePipelineBase> InItem, ESelectInfo::Type SelectInfo)
+void SInterchangePipelineConfigurationDialog::OnPipelineSelectionChanged(TSharedPtr<FInterchangePipelineItemType> InItem, ESelectInfo::Type SelectInfo)
 {
 	CurrentSelectedPipeline = nullptr;
 	if (InItem)
 	{
-		CurrentSelectedPipeline = InItem;
+		CurrentSelectedPipeline = InItem->Pipeline;
 	}
 	PipelineConfigurationDetailsView->SetObject(CurrentSelectedPipeline.Get());
 	
diff --git a/Engine/Plugins/Interchange/Editor/Source/Pipelines/Private/SInterchangePipelineConfigurationDialog.h b/Engine/Plugins/Interchange/Editor/Source/Pipelines/Private/SInterchangePipelineConfigurationDialog.h
--- a/Engine/Plugins/Interchange/Editor/Source/Pipelines/Private/SInterchangePipelineConfigurationDialog.h
+++ b/Engine/Plugins/Interchange/Editor/Source/Pipelines/Private/SInterchangePipelineConfigurationDialog.h
@@ -14,20 +14,28 @@ struct FSlateBrush;
 class IDetailsView;
 class SCheckBox;
 
-class SInterchangePipelineItem : public STableRow<TObjectPtr<UInterchangePipelineBase>>
+struct FInterchangePipelineItemType
+{
+public:
+	FString DisplayName;
+
+	UInterchangePipelineBase* Pipeline;
+};
+
+class SInterchangePipelineItem : public STableRow<TSharedPtr<FInterchangePipelineItemType>>
 {
 public:
 	void Construct(
 		const FArguments& InArgs,
 		const TSharedRef<STableViewBase>& OwnerTable,
-		TObjectPtr<UInterchangePipelineBase> InPipelineElement);
+		TSharedPtr<FInterchangePipelineItemType> InPipelineElement);
 private:
 	const FSlateBrush* GetImageItemIcon() const;
 
-	TObjectPtr<UInterchangePipelineBase> PipelineElement;
+	TSharedPtr<FInterchangePipelineItemType> PipelineElement;
 };
 
-typedef SListView< TObjectPtr<UInterchangePipelineBase> > SPipelineListViewType;
+typedef SListView< TSharedPtr<FInterchangePipelineItemType> > SPipelineListViewType;
 
 enum class ECloseEventType : uint8
 {
@@ -104,12 +112,12 @@ private:
 	// the pipelines list view
 	
 	TSharedPtr<SPipelineListViewType> PipelinesListView;
-	TArray< TObjectPtr<UInterchangePipelineBase> > PipelineListViewItems;
+	TArray< TSharedPtr< FInterchangePipelineItemType > > PipelineListViewItems;
 
 	/** list view generate row callback */
-	TSharedRef<ITableRow> MakePipelineListRowWidget(TObjectPtr<UInterchangePipelineBase> InElement, const TSharedRef<STableViewBase>& OwnerTable);
+	TSharedRef<ITableRow> MakePipelineListRowWidget(TSharedPtr<FInterchangePipelineItemType> InElement, const TSharedRef<STableViewBase>& OwnerTable);
 
-	void OnPipelineSelectionChanged(TObjectPtr<UInterchangePipelineBase> InItem, ESelectInfo::Type SelectInfo);
+	void OnPipelineSelectionChanged(TSharedPtr<FInterchangePipelineItemType> InItem, ESelectInfo::Type SelectInfo);
 
 	//
 	//////////////////////////////////////////////////////////////////////////
diff --git a/Engine/Source/Runtime/Interchange/Engine/Private/InterchangeManager.cpp b/Engine/Source/Runtime/Interchange/Engine/Private/InterchangeManager.cpp
--- a/Engine/Source/Runtime/Interchange/Engine/Private/InterchangeManager.cpp
+++ b/Engine/Source/Runtime/Interchange/Engine/Private/InterchangeManager.cpp
@@ -14,6 +14,7 @@
 #include "Engine/Blueprint.h"
 #include "Engine/World.h"
 #include "EngineAnalytics.h"
+#include "EngineUtils.h"
 #include "Framework/Notifications/NotificationManager.h"
 #include "InterchangeAssetImportData.h"
 #include "InterchangeBlueprintPipelineBase.h"
@@ -30,6 +31,7 @@
 #include "Misc/AsyncTaskNotification.h"
 #include "Misc/MessageDialog.h"
 #include "Misc/ScopedSlowTask.h"
+#include "Misc/DateTime.h"
 #include "Nodes/InterchangeBaseNodeContainer.h"
 #include "PackageUtils/PackageUtils.h"
 #include "Tasks/InterchangeTaskParsing.h"
@@ -557,6 +559,24 @@ void UE::Interchange::FImportAsyncHelper::CleanUp()
 		}
 	}
 	CreatedFactories.Empty();
+
+
+	for (TObjectIterator<UInterchangePipelineBase> PipelinesItr;PipelinesItr;++PipelinesItr)
+	{
+		const FString PackageName = PipelinesItr->GetPackage()->GetName();
+		if (PipelineInstancesPackageName == PackageName)
+		{
+			PipelinesItr->RemoveFromRoot();
+			PipelinesItr->MarkAsGarbage();
+		}
+	}
+
+	if (UPackage* PipelineInstancesPackage = FindPackage(nullptr, *PipelineInstancesPackageName))
+	{
+		PipelineInstancesPackage->ClearFlags(RF_Public| RF_Standalone| RF_Transactional);
+		PipelineInstancesPackage->SetFlags(RF_Transient);
+		PipelineInstancesPackage->MarkAsGarbage();
+	}
 }
 
 UE::Interchange::FImportResult::FImportResult()
@@ -725,8 +745,19 @@ void UE::Interchange::SanitizeObjectName(FString& ObjectName)
 	}
 }
 
-UInterchangePipelineBase* UE::Interchange::GeneratePipelineInstance(const FSoftObjectPath& PipelineInstance)
+UInterchangePipelineBase* UE::Interchange::GeneratePipelineInstanceInSourceAssetPackage(const FSoftObjectPath& SourcePipeline)
 {
+	FString PackageName = FPackageUtils::ExtractPackageName(SourcePipeline.ToString());
+	UPackage* TargetPackage = FindPackage(nullptr, *PackageName);
+	return GeneratePipelineInstance(SourcePipeline, TargetPackage);
+}
+
+UInterchangePipelineBase* UE::Interchange::GeneratePipelineInstance(const FSoftObjectPath& PipelineInstance, UPackage* PipelineInstancePackage /*= nullptr*/)
+{
+	if (!PipelineInstancePackage)
+	{
+		PipelineInstancePackage = GetTransientPackage();
+	}
 	UObject* ReferenceInstance = PipelineInstance.TryLoad();
 	if (!ReferenceInstance)
 	{
@@ -737,7 +768,7 @@ UInterchangePipelineBase* UE::Interchange::GeneratePipelineInstance(const FSoftO
 	{
 		if (BlueprintPipeline->GeneratedClass.Get())
 		{
-			GeneratedPipeline = NewObject<UInterchangePipelineBase>(GetTransientPackage(), BlueprintPipeline->GeneratedClass);
+			GeneratedPipeline = NewObject<UInterchangePipelineBase>(PipelineInstancePackage, BlueprintPipeline->GeneratedClass);
 		}
 		else
 		{
@@ -749,7 +780,7 @@ UInterchangePipelineBase* UE::Interchange::GeneratePipelineInstance(const FSoftO
 	{
 		if (PythonPipeline->GeneratedPipeline)
 		{
-			GeneratedPipeline = DuplicateObject<UInterchangePipelineBase>(PythonPipeline->GeneratedPipeline.Get(), GetTransientPackage());
+			GeneratedPipeline = DuplicateObject<UInterchangePipelineBase>(PythonPipeline->GeneratedPipeline.Get(), PipelineInstancePackage);
 		}
 		else
 		{
@@ -759,7 +790,7 @@ UInterchangePipelineBase* UE::Interchange::GeneratePipelineInstance(const FSoftO
 	}
 	else if (const UInterchangePipelineBase* DefaultPipeline = Cast<UInterchangePipelineBase>(ReferenceInstance))
 	{
-		GeneratedPipeline = DuplicateObject<UInterchangePipelineBase>(DefaultPipeline, GetTransientPackage());
+		GeneratedPipeline = DuplicateObject<UInterchangePipelineBase>(DefaultPipeline, PipelineInstancePackage);
 	}
 	else
 	{
@@ -1384,7 +1415,7 @@ UInterchangeManager::ImportInternal(const FString& ContentPath, const UInterchan
 		}
 	}
 #endif
-
+	
 	const bool bIsReimport = OriginalAssetImportData && OriginalAssetImportData->Pipelines.Num() > 0;
 	auto AdjustPipelineSettingForContext = [bImportScene, bIsReimport, ImportAssetParameters](UInterchangePipelineBase* Pipeline)
 	{
@@ -1427,6 +1458,15 @@ UInterchangeManager::ImportInternal(const FString& ContentPath, const UInterchan
 		TArray<FInterchangeStackInfo> PipelineStacks;
 		TArray<UInterchangePipelineBase*> OutPipelines;
 
+		FDateTime CurrentDateTime(FDateTime::Now());
+		const FString TransientPackageBasePath = GetTransientPackage()->GetPathName();
+		const FString PackageName = FString::Printf(TEXT("InterchangePipelinePackage_%llu"), CurrentDateTime.GetTicks());
+		AsyncHelper->PipelineInstancesPackageName = TransientPackageBasePath / PackageName;
+
+		UPackage* PipelineInstancesPackage =  CreatePackage(*AsyncHelper->PipelineInstancesPackageName);
+		PipelineInstancesPackage->ClearFlags(RF_Public | RF_Standalone);
+		PipelineInstancesPackage->SetPackageFlags(PKG_NewlyCreated);
+
 		//Fill the Stacks before showing the UI
 		if (bIsReimport)
 		{
@@ -1445,7 +1485,7 @@ UInterchangeManager::ImportInternal(const FString& ContentPath, const UInterchan
 				if (SourcePipeline) //Its possible a pipeline doesnt exist anymore so it wont load into memory when we loading the outer asset
 				{
 					//Duplicate the pipeline saved in the asset import data
-					UInterchangePipelineBase* GeneratedPipeline = Cast<UInterchangePipelineBase>(StaticDuplicateObject(SourcePipeline, GetTransientPackage()));
+					UInterchangePipelineBase* GeneratedPipeline = Cast<UInterchangePipelineBase>(StaticDuplicateObject(SourcePipeline, PipelineInstancesPackage));
 					// Make sure that the instance does not carry over standalone and public flags as they are not actual assets to be persisted
 					GeneratedPipeline->ClearFlags(EObjectFlags::RF_Standalone | EObjectFlags::RF_Public);
 					AdjustPipelineSettingForContext(GeneratedPipeline);
@@ -1485,7 +1525,7 @@ UInterchangeManager::ImportInternal(const FString& ContentPath, const UInterchan
 
 				for (int32 PipelineIndex = 0; PipelineIndex < Pipelines->Num(); ++PipelineIndex)
 				{
-					if (UInterchangePipelineBase* GeneratedPipeline = UE::Interchange::GeneratePipelineInstance((*Pipelines)[PipelineIndex]))
+					if (UInterchangePipelineBase* GeneratedPipeline = UE::Interchange::GeneratePipelineInstance((*Pipelines)[PipelineIndex], PipelineInstancesPackage))
 					{
 						AdjustPipelineSettingForContext(GeneratedPipeline);
 						StackInfo.Pipelines.Add(GeneratedPipeline);
diff --git a/Engine/Source/Runtime/Interchange/Engine/Public/InterchangeManager.h b/Engine/Source/Runtime/Interchange/Engine/Public/InterchangeManager.h
--- a/Engine/Source/Runtime/Interchange/Engine/Public/InterchangeManager.h
+++ b/Engine/Source/Runtime/Interchange/Engine/Public/InterchangeManager.h
@@ -217,6 +217,9 @@ namespace UE
 			FGraphEventRef PreCompletionTask;
 			FGraphEventRef CompletionTask;
 
+			// Package where the Pipeline Instances are stored during an import.
+			FString PipelineInstancesPackageName;
+
 			UPackage* GetCreatedPackage(const FString& PackageName) const;
 			void AddCreatedPackage(const FString& PackageName, UPackage* Package);
 
@@ -293,7 +296,9 @@ namespace UE
 		void SanitizeObjectName(FString& ObjectName);
 
 		/* This function take an asset representing a pipeline and generate a UInterchangePipelineBase asset. */
-		INTERCHANGEENGINE_API UInterchangePipelineBase* GeneratePipelineInstance(const FSoftObjectPath& PipelineInstance);
+		INTERCHANGEENGINE_API UInterchangePipelineBase* GeneratePipelineInstance(const FSoftObjectPath& PipelineInstance, UPackage* PipelineInstancePackage = nullptr);
+
+		INTERCHANGEENGINE_API UInterchangePipelineBase* GeneratePipelineInstanceInSourceAssetPackage(const FSoftObjectPath& PipelineInstance);
 
 	} //ns interchange
 } //ns UE
