[On behalf of Vedang Javdekar]
 Use the name of the pipeline assets in the import window instead of the class name
- The pipeline names in the Import Asset dialog using interchange pipelines have actual asset names followed by the class name in paranthesis.
- Fixed the issue of names displayed in the dialog on reimport were not matching the asset names.
- Fixed the case where the names were the same when a pipeline shows up in multiple pipeline stacks.

#rnx

diff --git a/Engine/Plugins/Interchange/Editor/Source/Pipelines/Private/SInterchangePipelineConfigurationDialog.cpp b/Engine/Plugins/Interchange/Editor/Source/Pipelines/Private/SInterchangePipelineConfigurationDialog.cpp
--- a/Engine/Plugins/Interchange/Editor/Source/Pipelines/Private/SInterchangePipelineConfigurationDialog.cpp
+++ b/Engine/Plugins/Interchange/Editor/Source/Pipelines/Private/SInterchangePipelineConfigurationDialog.cpp
@@ -25,6 +25,29 @@
 #define LOCTEXT_NAMESPACE "InterchangePipelineConfiguration"
 
 const FName ReimportStackName = TEXT("ReimportPipeline");
+const FString ReimportPipelinePrefix = TEXT("reimport_");
+
+ // Pipelines are renamed with the reimport prefix to avoid conflicts with the duplicates of the original pipelines that end up in the same package.
+ // As this is the name displayed in the Dialog, conflicts won't matter.
+FString SInterchangePipelineConfigurationDialog::GetPipelineDisplayName(const UInterchangePipelineBase* Pipeline)
+{
+	static int32 RightChopIndex = ReimportPipelinePrefix.Len();
+
+	FString PipelineDisplayName = Pipeline->GetName();
+	if (PipelineDisplayName.StartsWith(ReimportPipelinePrefix))
+	{
+		PipelineDisplayName = PipelineDisplayName.RightChop(RightChopIndex);
+	}
+
+	FString StackName;
+	FString DisplayName;
+	if (PipelineDisplayName.Split("_", &StackName, &DisplayName))
+	{
+		return DisplayName;
+	}
+
+	return PipelineDisplayName;
+}
 
 void SInterchangePipelineItem::Construct(
 	const FArguments& InArgs,
@@ -134,7 +157,7 @@ TSharedRef<SBox> SInterchangePipelineConfigurationDialog::SpawnPipelineConfigura
 			for (const TObjectPtr<UInterchangePipelineBase>& DefaultPipeline : Stack.Pipelines)
 			{
 				check(DefaultPipeline);
-				if (UInterchangePipelineBase* GeneratedPipeline = UE::Interchange::GeneratePipelineInstanceInSourceAssetPackage(DefaultPipeline))
+				if (UInterchangePipelineBase* GeneratedPipeline = UE::Interchange::GeneratePipelineInstanceInSourceAssetPackage(DefaultPipeline, CurrentStackName.ToString()))
 				{
 					GeneratedPipeline->TransferAdjustSettings(DefaultPipeline);
 					if (Stack.StackName == ReimportStackName)
@@ -150,7 +173,7 @@ TSharedRef<SBox> SInterchangePipelineConfigurationDialog::SpawnPipelineConfigura
 						GeneratedPipeline->PreDialogCleanup(Stack.StackName);
 					}
 
-					PipelineListViewItems.Add(MakeShareable(new FInterchangePipelineItemType{ DefaultPipeline.GetName(), GeneratedPipeline}));
+					PipelineListViewItems.Add(MakeShareable(new FInterchangePipelineItemType{ GetPipelineDisplayName(DefaultPipeline), GeneratedPipeline}));
 				}
 			}
 			SelectedStack = StackNamePtr;
@@ -482,7 +505,7 @@ FReply SInterchangePipelineConfigurationDialog::OnResetToDefault()
 
 							if (PipelineElement.Get() == Pipeline)
 							{
-								if (UInterchangePipelineBase* GeneratedPipeline = UE::Interchange::GeneratePipelineInstanceInSourceAssetPackage(DefaultPipeline))
+								if (UInterchangePipelineBase* GeneratedPipeline = UE::Interchange::GeneratePipelineInstanceInSourceAssetPackage(DefaultPipeline, CurrentStackName.ToString()))
 								{
 									GeneratedPipeline->TransferAdjustSettings(DefaultPipeline);
 									//Switch the pipeline the element point on
@@ -557,6 +580,12 @@ void SInterchangePipelineConfigurationDialog::ClosePipelineConfiguration(const E
 		//Fill the OutPipelines array
 		for (TSharedPtr<FInterchangePipelineItemType> PipelineElement : PipelineListViewItems)
 		{
+			if (!bReimport)
+			{
+				// Create a name that would not cause conflict when this asset maybe reimported.
+				FString NewPipelineName = ReimportPipelinePrefix + PipelineElement->DisplayName;
+				PipelineElement->Pipeline->Rename(*NewPipelineName);
+			}
 			OutPipelines->Add(PipelineElement->Pipeline);
 		}
 	}
@@ -640,7 +669,7 @@ void SInterchangePipelineConfigurationDialog::OnStackSelectionChanged(TSharedPtr
 		for (const TObjectPtr<UInterchangePipelineBase>& DefaultPipeline : Stack.Pipelines)
 		{
 			check(DefaultPipeline);
-			if (UInterchangePipelineBase* GeneratedPipeline = UE::Interchange::GeneratePipelineInstanceInSourceAssetPackage(DefaultPipeline))
+			if (UInterchangePipelineBase* GeneratedPipeline = UE::Interchange::GeneratePipelineInstanceInSourceAssetPackage(DefaultPipeline, CurrentStackName.ToString()))
 			{
 				GeneratedPipeline->TransferAdjustSettings(DefaultPipeline);
 				if (Stack.StackName != ReimportStackName)
@@ -649,7 +678,7 @@ void SInterchangePipelineConfigurationDialog::OnStackSelectionChanged(TSharedPtr
 					GeneratedPipeline->LoadSettings(Stack.StackName);
 					GeneratedPipeline->PreDialogCleanup(Stack.StackName);
 				}
-				PipelineListViewItems.Add(MakeShareable(new FInterchangePipelineItemType{ DefaultPipeline->GetName(), GeneratedPipeline}));
+				PipelineListViewItems.Add(MakeShareable(new FInterchangePipelineItemType{ GetPipelineDisplayName(DefaultPipeline), GeneratedPipeline}));
 			}
 		}
 	}
diff --git a/Engine/Plugins/Interchange/Editor/Source/Pipelines/Private/SInterchangePipelineConfigurationDialog.h b/Engine/Plugins/Interchange/Editor/Source/Pipelines/Private/SInterchangePipelineConfigurationDialog.h
--- a/Engine/Plugins/Interchange/Editor/Source/Pipelines/Private/SInterchangePipelineConfigurationDialog.h
+++ b/Engine/Plugins/Interchange/Editor/Source/Pipelines/Private/SInterchangePipelineConfigurationDialog.h
@@ -98,6 +98,9 @@ private:
 
 	void SaveAllPipelineSettings() const;
 
+	/** Internal utility function to properly display pipeline's name */
+	static FString GetPipelineDisplayName(const UInterchangePipelineBase* Pipeline);
+
 private:
 	TWeakPtr< SWindow > OwnerWindow;
 	TWeakObjectPtr<UInterchangeSourceData> SourceData;
diff --git a/Engine/Source/Runtime/Interchange/Engine/Private/InterchangeManager.cpp b/Engine/Source/Runtime/Interchange/Engine/Private/InterchangeManager.cpp
--- a/Engine/Source/Runtime/Interchange/Engine/Private/InterchangeManager.cpp
+++ b/Engine/Source/Runtime/Interchange/Engine/Private/InterchangeManager.cpp
@@ -67,11 +67,6 @@ namespace UE::Interchange::Private
 #endif
 	}
 
-	FInterchangeImportSettings& GetMutableImportSettings(UInterchangeProjectSettings& InterchangeProjectSettings, EImportType ImportType)
-	{
-		return const_cast<FInterchangeImportSettings&>(FInterchangeProjectSettingsUtils::GetImportSettings(InterchangeProjectSettings, ImportType == EImportType::ImportType_Scene));
-	}
-
 	void FillPipelineAnalyticData(UInterchangePipelineBase* Pipeline, const int32 UniqueId, const FString& ParentPipeline)
 	{
 		if (!FEngineAnalytics::IsAvailable())
@@ -746,14 +741,14 @@ void UE::Interchange::SanitizeObjectName(FString& ObjectName)
 	}
 }
 
-UInterchangePipelineBase* UE::Interchange::GeneratePipelineInstanceInSourceAssetPackage(const FSoftObjectPath& SourcePipeline)
+UInterchangePipelineBase* UE::Interchange::GeneratePipelineInstanceInSourceAssetPackage(const FSoftObjectPath& SourcePipeline, FString StackName)
 {
 	FString PackageName = FPackageUtils::ExtractPackageName(SourcePipeline.ToString());
 	UPackage* TargetPackage = FindPackage(nullptr, *PackageName);
-	return GeneratePipelineInstance(SourcePipeline, TargetPackage);
+	return GeneratePipelineInstance(SourcePipeline, StackName, TargetPackage);
 }
 
-UInterchangePipelineBase* UE::Interchange::GeneratePipelineInstance(const FSoftObjectPath& PipelineInstance, UPackage* PipelineInstancePackage /*= nullptr*/)
+UInterchangePipelineBase* UE::Interchange::GeneratePipelineInstance(const FSoftObjectPath& PipelineInstance, FString StackName, UPackage* PipelineInstancePackage /*= nullptr*/)
 {
 	if (!PipelineInstancePackage)
 	{
@@ -802,6 +797,9 @@ UInterchangePipelineBase* UE::Interchange::GeneratePipelineInstance(const FSoftO
 	{
 		// Make sure that the instance does not carry over standalone and public flags as they are not actual assets to be persisted
 		GeneratedPipeline->ClearFlags(EObjectFlags::RF_Standalone|EObjectFlags::RF_Public);
+		FString CurrentName = GeneratedPipeline->GetName();
+		FString NewName = FString::Printf(TEXT("%s_%s"), *StackName, *CurrentName);
+		GeneratedPipeline->Rename(*NewName,nullptr, REN_DoNotDirty | REN_NonTransactional);
 	}
 
 	return GeneratedPipeline;
@@ -1523,7 +1521,7 @@ UInterchangeManager::ImportInternal(const FString& ContentPath, const UInterchan
 
 				for (int32 PipelineIndex = 0; PipelineIndex < Pipelines->Num(); ++PipelineIndex)
 				{
-					if (UInterchangePipelineBase* GeneratedPipeline = UE::Interchange::GeneratePipelineInstance((*Pipelines)[PipelineIndex], PipelineInstancesPackage))
+					if (UInterchangePipelineBase* GeneratedPipeline = UE::Interchange::GeneratePipelineInstance((*Pipelines)[PipelineIndex], StackName.ToString(), PipelineInstancesPackage))
 					{
 						AdjustPipelineSettingForContext(GeneratedPipeline);
 						StackInfo.Pipelines.Add(GeneratedPipeline);
@@ -1641,7 +1639,8 @@ UInterchangeManager::ImportInternal(const FString& ContentPath, const UInterchan
 	{
 		for (int32 GraphPipelineIndex = 0; GraphPipelineIndex < ImportAssetParameters.OverridePipelines.Num(); ++GraphPipelineIndex)
 		{
-			UInterchangePipelineBase* GeneratedPipeline = UE::Interchange::GeneratePipelineInstance(ImportAssetParameters.OverridePipelines[GraphPipelineIndex]);
+			const FString OverridePiplinesStackName = TEXT("Overrides");
+			UInterchangePipelineBase* GeneratedPipeline = UE::Interchange::GeneratePipelineInstance(ImportAssetParameters.OverridePipelines[GraphPipelineIndex], OverridePiplinesStackName);
 			if (!GeneratedPipeline)
 			{
 				UE_LOG(LogInterchangeEngine, Error, TEXT("Interchange import: Override pipeline array contains a NULL pipeline. Script or code need to be fix to avoid this. "));
diff --git a/Engine/Source/Runtime/Interchange/Engine/Public/InterchangeManager.h b/Engine/Source/Runtime/Interchange/Engine/Public/InterchangeManager.h
--- a/Engine/Source/Runtime/Interchange/Engine/Public/InterchangeManager.h
+++ b/Engine/Source/Runtime/Interchange/Engine/Public/InterchangeManager.h
@@ -296,9 +296,9 @@ namespace UE
 		void SanitizeObjectName(FString& ObjectName);
 
 		/* This function take an asset representing a pipeline and generate a UInterchangePipelineBase asset. */
-		INTERCHANGEENGINE_API UInterchangePipelineBase* GeneratePipelineInstance(const FSoftObjectPath& PipelineInstance, UPackage* PipelineInstancePackage = nullptr);
+		INTERCHANGEENGINE_API UInterchangePipelineBase* GeneratePipelineInstance(const FSoftObjectPath& PipelineInstance, FString StackName, UPackage* PipelineInstancePackage = nullptr);
 
-		INTERCHANGEENGINE_API UInterchangePipelineBase* GeneratePipelineInstanceInSourceAssetPackage(const FSoftObjectPath& PipelineInstance);
+		INTERCHANGEENGINE_API UInterchangePipelineBase* GeneratePipelineInstanceInSourceAssetPackage(const FSoftObjectPath& PipelineInstance, FString StackName);
 
 	} //ns interchange
 } //ns UE
