Interchange: Allow translation to happen before openning the import dialog.
Filter the pipelines options with the translated data.

#rnx

diff --git a/Engine/Plugins/Interchange/Editor/Source/Pipelines/Private/InterchangePipelineConfigurationGeneric.cpp b/Engine/Plugins/Interchange/Editor/Source/Pipelines/Private/InterchangePipelineConfigurationGeneric.cpp
--- a/Engine/Plugins/Interchange/Editor/Source/Pipelines/Private/InterchangePipelineConfigurationGeneric.cpp
+++ b/Engine/Plugins/Interchange/Editor/Source/Pipelines/Private/InterchangePipelineConfigurationGeneric.cpp
@@ -2,6 +2,7 @@
 #include "InterchangePipelineConfigurationGeneric.h"
 
 #include "Framework/Application/SlateApplication.h"
+#include "Nodes/InterchangeBaseNodeContainer.h"
 #include "Interfaces/IMainFrameModule.h"
 #include "SInterchangePipelineConfigurationDialog.h"
 #include "Widgets/SWindow.h"
@@ -10,7 +11,8 @@
 
 EInterchangePipelineConfigurationDialogResult UInterchangePipelineConfigurationGeneric::ShowPipelineConfigurationDialog(TArray<FInterchangeStackInfo>& PipelineStacks
 	, TArray<UInterchangePipelineBase*>& OutPipelines
-	, TWeakObjectPtr<UInterchangeSourceData> SourceData)
+	, TWeakObjectPtr<UInterchangeSourceData> SourceData
+	, TWeakObjectPtr<UInterchangeBaseNodeContainer> BaseNodeContainer)
 {
 	//Create and show the graph inspector UI dialog
 	TSharedPtr<SWindow> ParentWindow;
@@ -32,6 +34,7 @@ EInterchangePipelineConfigurationDialogResult UInterchangePipelineConfigurationG
 		.bReimport(false)
 		.PipelineStacks(PipelineStacks)
 		.OutPipelines(&OutPipelines)
+		.BaseNodeContainer(BaseNodeContainer)
 	);
 
 	FSlateApplication::Get().AddModalWindow(Window, ParentWindow, false);
@@ -51,7 +54,8 @@ EInterchangePipelineConfigurationDialogResult UInterchangePipelineConfigurationG
 
 EInterchangePipelineConfigurationDialogResult UInterchangePipelineConfigurationGeneric::ShowScenePipelineConfigurationDialog(TArray<FInterchangeStackInfo>& PipelineStacks
 	, TArray<UInterchangePipelineBase*>& OutPipelines
-	, TWeakObjectPtr<UInterchangeSourceData> SourceData)
+	, TWeakObjectPtr<UInterchangeSourceData> SourceData
+	, TWeakObjectPtr<UInterchangeBaseNodeContainer> BaseNodeContainer)
 {
 	//Create and show the graph inspector UI dialog
 	TSharedPtr<SWindow> ParentWindow;
@@ -73,6 +77,7 @@ EInterchangePipelineConfigurationDialogResult UInterchangePipelineConfigurationG
 		.bReimport(false)
 		.PipelineStacks(PipelineStacks)
 		.OutPipelines(&OutPipelines)
+		.BaseNodeContainer(BaseNodeContainer)
 	);
 
 	FSlateApplication::Get().AddModalWindow(Window, ParentWindow, false);
@@ -92,7 +97,8 @@ EInterchangePipelineConfigurationDialogResult UInterchangePipelineConfigurationG
 
 EInterchangePipelineConfigurationDialogResult UInterchangePipelineConfigurationGeneric::ShowReimportPipelineConfigurationDialog(TArray<FInterchangeStackInfo>& PipelineStacks
 	, TArray<UInterchangePipelineBase*>& OutPipelines
-	, TWeakObjectPtr<UInterchangeSourceData> SourceData)
+	, TWeakObjectPtr<UInterchangeSourceData> SourceData
+	, TWeakObjectPtr<UInterchangeBaseNodeContainer> BaseNodeContainer)
 {
 	//Create and show the graph inspector UI dialog
 	TSharedPtr<SWindow> ParentWindow;
@@ -114,6 +120,7 @@ EInterchangePipelineConfigurationDialogResult UInterchangePipelineConfigurationG
 		.bReimport(true)
 		.PipelineStacks(PipelineStacks)
 		.OutPipelines(&OutPipelines)
+		.BaseNodeContainer(BaseNodeContainer)
 	);
 
 	FSlateApplication::Get().AddModalWindow(Window, ParentWindow, false);
diff --git a/Engine/Plugins/Interchange/Editor/Source/Pipelines/Private/SInterchangePipelineConfigurationDialog.cpp b/Engine/Plugins/Interchange/Editor/Source/Pipelines/Private/SInterchangePipelineConfigurationDialog.cpp
--- a/Engine/Plugins/Interchange/Editor/Source/Pipelines/Private/SInterchangePipelineConfigurationDialog.cpp
+++ b/Engine/Plugins/Interchange/Editor/Source/Pipelines/Private/SInterchangePipelineConfigurationDialog.cpp
@@ -8,6 +8,7 @@
 #include "Framework/Views/TableViewMetadata.h"
 #include "IDetailsView.h"
 #include "IDocumentation.h"
+#include "Nodes/InterchangeBaseNodeContainer.h"
 #include "InterchangeManager.h"
 #include "InterchangePipelineConfigurationBase.h"
 #include "InterchangeProjectSettings.h"
@@ -172,7 +173,10 @@ TSharedRef<SBox> SInterchangePipelineConfigurationDialog::SpawnPipelineConfigura
 						GeneratedPipeline->LoadSettings(Stack.StackName);
 						GeneratedPipeline->PreDialogCleanup(Stack.StackName);
 					}
-
+					if (bFilterOptions && BaseNodeContainer.IsValid())
+					{
+						GeneratedPipeline->FilterPropertiesFromTranslatedData(BaseNodeContainer.Get());
+					}
 					PipelineListViewItems.Add(MakeShareable(new FInterchangePipelineItemType{ GetPipelineDisplayName(DefaultPipeline), GeneratedPipeline}));
 				}
 			}
@@ -285,6 +289,15 @@ TSharedRef<SBox> SInterchangePipelineConfigurationDialog::SpawnPipelineConfigura
 		{
 			return IsPropertyVisible(PropertyAndParent);
 		});
+	PipelineConfigurationDetailsView->OnFinishedChangingProperties().AddLambda([this](const FPropertyChangedEvent& PropertyChangedEvent)
+	{
+		if (CurrentSelectedPipeline && CurrentSelectedPipeline->IsPropertyChangeNeedRefresh(PropertyChangedEvent))
+		{
+			//Refresh the pipeline
+			constexpr bool bStackSelectionChange = false;
+			RefreshStack(bStackSelectionChange);
+		}
+	});
 	return PipelineConfigurationPanelBox;
 }
 
@@ -298,6 +311,7 @@ void SInterchangePipelineConfigurationDialog::Construct(const FArguments& InArgs
 	bReimport = InArgs._bReimport;
 	PipelineStacks = InArgs._PipelineStacks;
 	OutPipelines = InArgs._OutPipelines;
+	BaseNodeContainer = InArgs._BaseNodeContainer;
 
 	check(OutPipelines);
 
@@ -307,6 +321,15 @@ void SInterchangePipelineConfigurationDialog::Construct(const FArguments& InArgs
 		OwnerWindowPinned->GetOnWindowClosedEvent().AddRaw(this, &SInterchangePipelineConfigurationDialog::OnWindowClosed);
 	}
 
+	if (bReimport)
+	{
+		bFilterOptions = false;
+	}
+	else if(GConfig->DoesSectionExist(TEXT("InterchangeImportDialogOptions"), GEditorPerProjectIni))
+	{
+		GConfig->GetBool(TEXT("InterchangeImportDialogOptions"), TEXT("FilterOptions"), bFilterOptions, GEditorPerProjectIni);
+	}
+
 	this->ChildSlot
 	[
 		SNew(SBorder)
@@ -331,6 +354,34 @@ void SInterchangePipelineConfigurationDialog::Construct(const FArguments& InArgs
 				.AutoWidth()
 				.HAlign(HAlign_Right)
 				.Padding(10.0f, 2.0f, 0.0f, 2.0f)
+				[
+					SNew(SHorizontalBox)
+					.ToolTipText(LOCTEXT("SInterchangePipelineConfigurationDialog_FilterPipelineOptions_tooltip", "Filter the pipeline options using the source content data."))
+					.Visibility_Lambda([this]()
+						{
+							return bReimport ? EVisibility::Collapsed : EVisibility::All;
+						})
+					+ SHorizontalBox::Slot()
+					.AutoWidth()
+					.HAlign(HAlign_Right)
+					.VAlign(VAlign_Center)
+					.Padding(4.f, 0.f)
+					[
+						SNew(STextBlock)
+						.Text(LOCTEXT("SInterchangePipelineConfigurationDialog_FilterPipelineOptions", "Filter Options"))
+					]
+					+ SHorizontalBox::Slot()
+					.Padding(4.f, 0.f)
+					[
+						SNew(SCheckBox)
+						.IsChecked(this, &SInterchangePipelineConfigurationDialog::IsFilteringOptions)
+						.OnCheckStateChanged(this, &SInterchangePipelineConfigurationDialog::OnFilterOptionsChanged)
+					]
+				]
+				+SHorizontalBox::Slot()
+				.AutoWidth()
+				.HAlign(HAlign_Right)
+				.Padding(2.0f, 2.0f, 0.0f, 2.0f)
 				[
 					SNew(SButton)
 					.HAlign(HAlign_Center)
@@ -508,6 +559,10 @@ FReply SInterchangePipelineConfigurationDialog::OnResetToDefault()
 								if (UInterchangePipelineBase* GeneratedPipeline = UE::Interchange::GeneratePipelineInstanceInSourceAssetPackage(DefaultPipeline))
 								{
 									GeneratedPipeline->TransferAdjustSettings(DefaultPipeline);
+									if(bFilterOptions && BaseNodeContainer.IsValid())
+									{
+										GeneratedPipeline->FilterPropertiesFromTranslatedData(BaseNodeContainer.Get());
+									}
 									//Switch the pipeline the element point on
 									PipelineListViewItems[PipelineIndex]->Pipeline = GeneratedPipeline;
 									PipelineConfigurationDetailsView->SetObject(GeneratedPipeline, true);
@@ -625,25 +680,8 @@ FReply SInterchangePipelineConfigurationDialog::OnKeyDown(const FGeometry& MyGeo
 	return FReply::Unhandled();
 }
 
-void SInterchangePipelineConfigurationDialog::OnStackSelectionChanged(TSharedPtr<FString> String, ESelectInfo::Type)
+void SInterchangePipelineConfigurationDialog::RefreshStack(bool bStackSelectionChange)
 {
-	if (!String.IsValid())
-	{
-		return;
-	}
-
-	FName NewStackName = FName(*String.Get());
-	if (!UE::Private::ContainStack(PipelineStacks, NewStackName))
-	{
-		return;
-	}
-
-	//Nothing change the selection is the same
-	if (CurrentStackName == NewStackName)
-	{
-		return;
-	}
-
 	//Save current stack settings, we want the same settings when we will go back to the same stack
 	//When doing a reimport we do not want to save the setting because the context have special default
 	//value for some options like: (Import Materials, Import Textures...).
@@ -653,9 +691,21 @@ void SInterchangePipelineConfigurationDialog::OnStackSelectionChanged(TSharedPtr
 		SaveAllPipelineSettings();
 	}
 
-	//Use the stack select by interchange manager
-	CurrentStackName = NewStackName;
-	
+	int32 CurrentPipelineIndex = 0;
+	if (!bStackSelectionChange)
+	{
+		//store the selected pipeline
+		for (int32 PipelineIndex = 0; PipelineIndex < PipelineListViewItems.Num(); ++PipelineIndex)
+		{
+			const TSharedPtr<FInterchangePipelineItemType> PipelineItem = PipelineListViewItems[PipelineIndex];
+			if (PipelineItem->Pipeline == CurrentSelectedPipeline)
+			{
+				CurrentPipelineIndex = PipelineIndex;
+				break;
+			}
+		}
+	}
+
 	//Rebuild the Pipeline list item
 	PipelineListViewItems.Reset();
 
@@ -678,7 +728,11 @@ void SInterchangePipelineConfigurationDialog::OnStackSelectionChanged(TSharedPtr
 					GeneratedPipeline->LoadSettings(Stack.StackName);
 					GeneratedPipeline->PreDialogCleanup(Stack.StackName);
 				}
-				PipelineListViewItems.Add(MakeShareable(new FInterchangePipelineItemType{ GetPipelineDisplayName(DefaultPipeline), GeneratedPipeline}));
+				if (bFilterOptions && BaseNodeContainer.IsValid())
+				{
+					GeneratedPipeline->FilterPropertiesFromTranslatedData(BaseNodeContainer.Get());
+				}
+				PipelineListViewItems.Add(MakeShareable(new FInterchangePipelineItemType{ GetPipelineDisplayName(DefaultPipeline), GeneratedPipeline }));
 			}
 		}
 	}
@@ -686,11 +740,38 @@ void SInterchangePipelineConfigurationDialog::OnStackSelectionChanged(TSharedPtr
 	//Select the first pipeline
 	if (PipelineListViewItems.Num() > 0)
 	{
-		PipelinesListView->SetSelection(PipelineListViewItems[0], ESelectInfo::Direct);
+		CurrentPipelineIndex = PipelineListViewItems.IsValidIndex(CurrentPipelineIndex) ? CurrentPipelineIndex : 0;
+		PipelinesListView->SetSelection(PipelineListViewItems[CurrentPipelineIndex], ESelectInfo::Direct);
 	}
 	PipelinesListView->RequestListRefresh();
 }
 
+void SInterchangePipelineConfigurationDialog::OnStackSelectionChanged(TSharedPtr<FString> String, ESelectInfo::Type)
+{
+	if (!String.IsValid())
+	{
+		return;
+	}
+
+	FName NewStackName = FName(*String.Get());
+	if (!UE::Private::ContainStack(PipelineStacks, NewStackName))
+	{
+		return;
+	}
+
+	//Nothing change the selection is the same
+	if (CurrentStackName == NewStackName)
+	{
+		return;
+	}
+
+	//Use the stack select by interchange manager
+	CurrentStackName = NewStackName;
+
+	constexpr bool bStackSelectionChange = true;
+	RefreshStack(bStackSelectionChange);
+}
+
 TSharedRef<ITableRow> SInterchangePipelineConfigurationDialog::MakePipelineListRowWidget(
 	TSharedPtr<FInterchangePipelineItemType> InElement,
 	const TSharedRef<STableViewBase>& OwnerTable)
@@ -706,6 +787,7 @@ void SInterchangePipelineConfigurationDialog::OnPipelineSelectionChanged(TShared
 	{
 		CurrentSelectedPipeline = InItem->Pipeline;
 	}
+	
 	PipelineConfigurationDetailsView->SetObject(CurrentSelectedPipeline.Get());
 	
 	if (CurrentSelectedPipeline)
@@ -716,4 +798,20 @@ void SInterchangePipelineConfigurationDialog::OnPipelineSelectionChanged(TShared
 	}
 }
 
+void SInterchangePipelineConfigurationDialog::OnFilterOptionsChanged(ECheckBoxState CheckState)
+{
+	bool bNewCheckValue = CheckState == ECheckBoxState::Checked ? true : false;
+	if (bNewCheckValue == bFilterOptions)
+	{
+		//Check state did not change
+		return;
+	}
+	bFilterOptions = bNewCheckValue;
+	//Refresh the pipeline
+	constexpr bool bStackSelectionChange = false;
+	RefreshStack(bStackSelectionChange);
+
+	GConfig->SetBool(TEXT("InterchangeImportDialogOptions"), TEXT("FilterOptions"), bFilterOptions, GEditorPerProjectIni);
+}
+
 #undef LOCTEXT_NAMESPACE
diff --git a/Engine/Plugins/Interchange/Editor/Source/Pipelines/Private/SInterchangePipelineConfigurationDialog.h b/Engine/Plugins/Interchange/Editor/Source/Pipelines/Private/SInterchangePipelineConfigurationDialog.h
--- a/Engine/Plugins/Interchange/Editor/Source/Pipelines/Private/SInterchangePipelineConfigurationDialog.h
+++ b/Engine/Plugins/Interchange/Editor/Source/Pipelines/Private/SInterchangePipelineConfigurationDialog.h
@@ -2,6 +2,7 @@
 
 #pragma once
 
+#include "Nodes/InterchangeBaseNodeContainer.h"
 #include "InterchangePipelineBase.h"
 #include "InterchangePipelineConfigurationBase.h"
 #include "Styling/SlateBrush.h"
@@ -57,6 +58,7 @@ public:
 		SLATE_ARGUMENT(bool, bReimport)
 		SLATE_ARGUMENT(TArray<FInterchangeStackInfo>, PipelineStacks)
 		SLATE_ARGUMENT(TArray<UInterchangePipelineBase*>*, OutPipelines)
+		SLATE_ARGUMENT(TWeakObjectPtr<UInterchangeBaseNodeContainer>, BaseNodeContainer)
 	SLATE_END_ARGS()
 
 public:
@@ -104,12 +106,14 @@ private:
 private:
 	TWeakPtr< SWindow > OwnerWindow;
 	TWeakObjectPtr<UInterchangeSourceData> SourceData;
+	TWeakObjectPtr<UInterchangeBaseNodeContainer> BaseNodeContainer;
 	TArray<FInterchangeStackInfo> PipelineStacks;
 	TArray<UInterchangePipelineBase*>* OutPipelines;
 
 	// The available stacks
 	TArray<TSharedPtr<FString>> AvailableStacks;
 	void OnStackSelectionChanged(TSharedPtr<FString> String, ESelectInfo::Type);
+	void RefreshStack(bool bStackSelectionChange);
 
 	//////////////////////////////////////////////////////////////////////////
 	// the pipelines list view
@@ -125,6 +129,13 @@ private:
 	//
 	//////////////////////////////////////////////////////////////////////////
 
+	ECheckBoxState IsFilteringOptions() const
+	{
+		return bFilterOptions ? ECheckBoxState::Checked : ECheckBoxState::Unchecked;
+	}
+
+	void OnFilterOptionsChanged(ECheckBoxState CheckState);
+
 	TSharedPtr<IDetailsView> PipelineConfigurationDetailsView;
 	TSharedPtr<SCheckBox> UseSameSettingsForAllCheckBox;
 
@@ -133,6 +144,8 @@ private:
 	bool bCanceled = false;
 	bool bImportAll = false;
 
+	bool bFilterOptions = false;
+
 	FName CurrentStackName = NAME_None;
 	TObjectPtr<UInterchangePipelineBase> CurrentSelectedPipeline = nullptr;
 };
diff --git a/Engine/Plugins/Interchange/Editor/Source/Pipelines/Public/InterchangePipelineConfigurationGeneric.h b/Engine/Plugins/Interchange/Editor/Source/Pipelines/Public/InterchangePipelineConfigurationGeneric.h
--- a/Engine/Plugins/Interchange/Editor/Source/Pipelines/Public/InterchangePipelineConfigurationGeneric.h
+++ b/Engine/Plugins/Interchange/Editor/Source/Pipelines/Public/InterchangePipelineConfigurationGeneric.h
@@ -6,6 +6,7 @@
 
 #include "InterchangePipelineConfigurationGeneric.generated.h"
 
+class UInterchangeBaseNodeContainer;
 class UInterchangeSourceData;
 
 UCLASS(BlueprintType, editinlinenew, Experimental)
@@ -19,13 +20,16 @@ protected:
 
 	virtual EInterchangePipelineConfigurationDialogResult ShowPipelineConfigurationDialog(TArray<FInterchangeStackInfo>& PipelineStacks
 		, TArray<UInterchangePipelineBase*>& OutPipelines
-		, TWeakObjectPtr<UInterchangeSourceData> SourceData) override;
+		, TWeakObjectPtr<UInterchangeSourceData> SourceData
+		, TWeakObjectPtr<UInterchangeBaseNodeContainer> BaseNodeContainer) override;
 	virtual EInterchangePipelineConfigurationDialogResult ShowScenePipelineConfigurationDialog(TArray<FInterchangeStackInfo>& PipelineStacks
 		, TArray<UInterchangePipelineBase*>& OutPipelines
-		, TWeakObjectPtr<UInterchangeSourceData> SourceData) override;
+		, TWeakObjectPtr<UInterchangeSourceData> SourceData
+		, TWeakObjectPtr<UInterchangeBaseNodeContainer> BaseNodeContainer) override;
 	virtual EInterchangePipelineConfigurationDialogResult ShowReimportPipelineConfigurationDialog(TArray<FInterchangeStackInfo>& PipelineStacks
 		, TArray<UInterchangePipelineBase*>& OutPipelines
-		, TWeakObjectPtr<UInterchangeSourceData> SourceData) override;
+		, TWeakObjectPtr<UInterchangeSourceData> SourceData
+		, TWeakObjectPtr<UInterchangeBaseNodeContainer> BaseNodeContainer) override;
 };
 
 #if UE_ENABLE_INCLUDE_ORDER_DEPRECATED_IN_5_2
diff --git a/Engine/Plugins/Interchange/Runtime/Source/Pipelines/Private/InterchangeGenericAnimationPipeline.cpp b/Engine/Plugins/Interchange/Runtime/Source/Pipelines/Private/InterchangeGenericAnimationPipeline.cpp
--- a/Engine/Plugins/Interchange/Runtime/Source/Pipelines/Private/InterchangeGenericAnimationPipeline.cpp
+++ b/Engine/Plugins/Interchange/Runtime/Source/Pipelines/Private/InterchangeGenericAnimationPipeline.cpp
@@ -25,6 +25,7 @@ void UInterchangeGenericAnimationPipeline::AdjustSettingsForContext(EInterchange
 {
 	Super::AdjustSettingsForContext(ImportType, ReimportAsset);
 
+#if WITH_EDITOR
 	check(CommonSkeletalMeshesAndAnimationsProperties.IsValid());
 	
 	bSceneImport = ImportType == EInterchangePipelineContext::SceneImport
@@ -61,6 +62,7 @@ void UInterchangeGenericAnimationPipeline::AdjustSettingsForContext(EInterchange
 			HidePropertiesOfCategory(OuterMostPipeline, this, HideCategoryName);
 		}
 	}
+#endif //WITH_EDITOR
 }
 
 void UInterchangeGenericAnimationPipeline::ExecutePipeline(UInterchangeBaseNodeContainer* InBaseNodeContainer, const TArray<UInterchangeSourceData*>& InSourceDatas)
diff --git a/Engine/Plugins/Interchange/Runtime/Source/Pipelines/Private/InterchangeGenericAssetsPipeline.cpp b/Engine/Plugins/Interchange/Runtime/Source/Pipelines/Private/InterchangeGenericAssetsPipeline.cpp
--- a/Engine/Plugins/Interchange/Runtime/Source/Pipelines/Private/InterchangeGenericAssetsPipeline.cpp
+++ b/Engine/Plugins/Interchange/Runtime/Source/Pipelines/Private/InterchangeGenericAssetsPipeline.cpp
@@ -4,6 +4,7 @@
 
 #include "Animation/Skeleton.h"
 #include "CoreMinimal.h"
+#include "InterchangeAnimationTrackSetNode.h"
 #include "InterchangeAnimSequenceFactoryNode.h"
 #include "InterchangeCommonPipelineDataFactoryNode.h"
 #include "InterchangeGenericAnimationPipeline.h"
@@ -120,6 +121,190 @@ void UInterchangeGenericAssetsPipeline::AdjustSettingsForContext(EInterchangePip
 }
 
 #if WITH_EDITOR
+
+void UInterchangeGenericAssetsPipeline::FilterPropertiesFromTranslatedData(UInterchangeBaseNodeContainer* InBaseNodeContainer)
+{
+	Super::FilterPropertiesFromTranslatedData(InBaseNodeContainer);
+
+	if (MaterialPipeline)
+	{
+		MaterialPipeline->FilterPropertiesFromTranslatedData(InBaseNodeContainer);
+	}
+
+	if (CommonMeshesProperties && CommonSkeletalMeshesAndAnimationsProperties && MeshPipeline && AnimationPipeline)
+	{
+		CommonMeshesProperties->FilterPropertiesFromTranslatedData(InBaseNodeContainer);
+		CommonSkeletalMeshesAndAnimationsProperties->FilterPropertiesFromTranslatedData(InBaseNodeContainer);
+		MeshPipeline->FilterPropertiesFromTranslatedData(InBaseNodeContainer);
+		AnimationPipeline->FilterPropertiesFromTranslatedData(InBaseNodeContainer);
+
+		UInterchangePipelineMeshesUtilities* PipelineMeshesUtilities = UInterchangeGenericMeshPipeline::CreateMeshPipelineUtilities(InBaseNodeContainer, MeshPipeline);
+
+		TArray<FString> SkeletalMeshes;
+		PipelineMeshesUtilities->GetAllSkinnedMeshInstance(SkeletalMeshes);
+		if(SkeletalMeshes.Num() == 0)
+		{
+			PipelineMeshesUtilities->GetAllSkinnedMeshGeometry(SkeletalMeshes);
+		}
+		TArray<FString> StaticMeshes;
+		PipelineMeshesUtilities->GetAllStaticMeshInstance(StaticMeshes);
+		if(StaticMeshes.Num() == 0)
+		{
+			PipelineMeshesUtilities->GetAllStaticMeshGeometry(StaticMeshes);
+		}
+
+		int32 RawStaticMesh = 0;
+		int32 RawSkeletalMesh = 0;
+		int32 RawMorphTargetShape = 0;
+		InBaseNodeContainer->IterateNodesOfType<UInterchangeMeshNode>([&RawStaticMesh, &RawSkeletalMesh, &RawMorphTargetShape](const FString& NodeUid, UInterchangeMeshNode* MeshNode)
+			{
+				if (MeshNode->IsMorphTarget())
+				{
+					RawMorphTargetShape++;
+				}
+				else
+				{
+					MeshNode->IsSkinnedMesh() ? RawSkeletalMesh++ : RawStaticMesh++;
+				}
+			});
+
+		int32 RawAnimationNode = 0;
+		InBaseNodeContainer->IterateNodesOfType<UInterchangeAnimationTrackBaseNode>([&RawAnimationNode](const FString& NodeUid, UInterchangeAnimationTrackBaseNode* AnimationNode)
+			{
+				RawAnimationNode++;
+			});
+		InBaseNodeContainer->IterateNodesOfType<UInterchangeAnimationTrackSetNode>([&RawAnimationNode](const FString& NodeUid, UInterchangeAnimationTrackSetNode* AnimationNode)
+			{
+				RawAnimationNode++;
+			});
+
+		UInterchangePipelineBase* OuterMostPipeline = GetMostPipelineOuter();
+		auto HideFullCategory = [this, &OuterMostPipeline](const FString& Category, UInterchangePipelineBase* Pipeline)
+		{
+			TArray<FString> HideCategories;
+			//Filter out all Textures properties
+			HideCategories.Add(Category);
+			if (OuterMostPipeline)
+			{
+				for (const FString& HideCategoryName : HideCategories)
+				{
+					HidePropertiesOfCategory(OuterMostPipeline, Pipeline, HideCategoryName);
+				}
+			}
+		};
+
+		auto LocalHideProperty = [this, &OuterMostPipeline](UInterchangePipelineBase* Pipeline, FName PropertyName)
+		{
+			if (OuterMostPipeline)
+			{
+				HideProperty(OuterMostPipeline, Pipeline, PropertyName);
+			}
+		};
+
+		//Found which categories to hide
+		bool bHideStaticMeshes = false;
+		bool bHideSkeletalMeshes = false;
+		bool bHideCommonMeshes = false;
+		bool bHideCommonSkeletalMeshesAndAnimations = false;
+		bool bHideCommonSkeletalMeshesAndAnimations_StaticMesh = false;
+		bool bHideAnimations = false;
+
+		if (RawStaticMesh == 0 || RawMorphTargetShape == 0)
+		{
+			bHideCommonSkeletalMeshesAndAnimations_StaticMesh = true;
+		}
+
+		if (SkeletalMeshes.Num() == 0 && StaticMeshes.Num() == 0)
+		{
+			bHideStaticMeshes = true;
+			bHideSkeletalMeshes = true;
+			bHideCommonMeshes = true;
+		}
+		else if (StaticMeshes.Num() > 0 && SkeletalMeshes.Num() == 0)
+		{
+			bHideSkeletalMeshes = true;
+			bHideCommonSkeletalMeshesAndAnimations = true;
+			bHideAnimations = true;
+		}
+		else if (SkeletalMeshes.Num() > 0 && StaticMeshes.Num() == 0)
+		{
+			bHideStaticMeshes = true;
+		}
+		else
+		{
+			constexpr bool bFilterCodeLogicError = false;
+			ensure(bFilterCodeLogicError);
+		}
+
+		if (SkeletalMeshes.Num() > 0)
+		{
+			if (MeshPipeline->SkeletalMeshImportContentType == EInterchangeSkeletalMeshContentType::SkinningWeights)
+			{
+				LocalHideProperty(CommonMeshesProperties, GET_MEMBER_NAME_CHECKED(UInterchangeGenericCommonMeshesProperties, VertexOverrideColor));
+				LocalHideProperty(CommonMeshesProperties, GET_MEMBER_NAME_CHECKED(UInterchangeGenericCommonMeshesProperties, VertexColorImportOption));
+				LocalHideProperty(MeshPipeline, GET_MEMBER_NAME_CHECKED(UInterchangeGenericMeshPipeline, bImportMorphTargets));
+				LocalHideProperty(MeshPipeline, GET_MEMBER_NAME_CHECKED(UInterchangeGenericMeshPipeline, ThresholdPosition));
+				LocalHideProperty(MeshPipeline, GET_MEMBER_NAME_CHECKED(UInterchangeGenericMeshPipeline, ThresholdTangentNormal));
+				LocalHideProperty(MeshPipeline, GET_MEMBER_NAME_CHECKED(UInterchangeGenericMeshPipeline, ThresholdUV));
+				LocalHideProperty(MeshPipeline, GET_MEMBER_NAME_CHECKED(UInterchangeGenericMeshPipeline, MorphThresholdPosition));
+			}
+			else if (MeshPipeline->SkeletalMeshImportContentType == EInterchangeSkeletalMeshContentType::Geometry)
+			{
+				LocalHideProperty(MeshPipeline, GET_MEMBER_NAME_CHECKED(UInterchangeGenericMeshPipeline, bUpdateSkeletonReferencePose));
+			}
+		}
+
+		if (RawAnimationNode == 0)
+		{
+			bHideAnimations = true;
+		}
+
+		if (bHideAnimations && bHideSkeletalMeshes)
+		{
+			bHideCommonSkeletalMeshesAndAnimations = true;
+		}
+
+		//Hide the categories
+		if (bHideStaticMeshes)
+		{
+			HideFullCategory(TEXT("Static Meshes"), MeshPipeline);
+		}
+		if (bHideSkeletalMeshes)
+		{
+			HideFullCategory(TEXT("Skeletal Meshes"), MeshPipeline);
+		}
+		if(bHideCommonMeshes)
+		{
+			HideFullCategory(TEXT("Common Meshes"), CommonMeshesProperties);
+		}
+		if (bHideCommonSkeletalMeshesAndAnimations)
+		{
+			HideFullCategory(TEXT("Common Skeletal Meshes and Animations"), CommonSkeletalMeshesAndAnimationsProperties);
+		}
+		if (bHideCommonSkeletalMeshesAndAnimations_StaticMesh)
+		{
+			HideFullCategory(TEXT("Static Meshes"), CommonSkeletalMeshesAndAnimationsProperties);
+		}
+		if (bHideAnimations)
+		{
+			HideFullCategory(TEXT("Animations"), AnimationPipeline);
+		}
+	}
+}
+
+bool UInterchangeGenericAssetsPipeline::IsPropertyChangeNeedRefresh(const FPropertyChangedEvent& PropertyChangedEvent)
+{
+	if ((CommonMeshesProperties && CommonMeshesProperties->IsPropertyChangeNeedRefresh(PropertyChangedEvent))
+		|| (CommonSkeletalMeshesAndAnimationsProperties && CommonSkeletalMeshesAndAnimationsProperties->IsPropertyChangeNeedRefresh(PropertyChangedEvent))
+		|| (MeshPipeline && MeshPipeline->IsPropertyChangeNeedRefresh(PropertyChangedEvent))
+		|| (MaterialPipeline && MaterialPipeline->IsPropertyChangeNeedRefresh(PropertyChangedEvent))
+		|| (AnimationPipeline && AnimationPipeline->IsPropertyChangeNeedRefresh(PropertyChangedEvent)))
+	{
+		return true;
+	}
+	return Super::IsPropertyChangeNeedRefresh(PropertyChangedEvent);
+}
+
 bool UInterchangeGenericAssetsPipeline::GetPropertyPossibleValues(const FName PropertyPath, TArray<FString>& PossibleValues)
 {
 	const FString PropertyPathString = PropertyPath.ToString();
@@ -150,7 +335,8 @@ bool UInterchangeGenericAssetsPipeline::GetPropertyPossibleValues(const FName Pr
 	//If we did not find any property call the super implementation
 	return Super::GetPropertyPossibleValues(PropertyPath, PossibleValues);
 }
-#endif
+
+#endif //WITH_EDITOR
 
 void UInterchangeGenericAssetsPipeline::ExecutePipeline(UInterchangeBaseNodeContainer* InBaseNodeContainer, const TArray<UInterchangeSourceData*>& InSourceDatas)
 {
diff --git a/Engine/Plugins/Interchange/Runtime/Source/Pipelines/Private/InterchangeGenericMaterialPipeline.cpp b/Engine/Plugins/Interchange/Runtime/Source/Pipelines/Private/InterchangeGenericMaterialPipeline.cpp
--- a/Engine/Plugins/Interchange/Runtime/Source/Pipelines/Private/InterchangeGenericMaterialPipeline.cpp
+++ b/Engine/Plugins/Interchange/Runtime/Source/Pipelines/Private/InterchangeGenericMaterialPipeline.cpp
@@ -487,7 +487,7 @@ void UInterchangeGenericMaterialPipeline::AdjustSettingsForContext(EInterchangeP
 	{
 		TexturePipeline->AdjustSettingsForContext(ImportType, ReimportAsset);
 	}
-
+#if WITH_EDITOR
 	TArray<FString> HideCategories;
 	bool bIsObjectAMaterial = !ReimportAsset ? false : ReimportAsset->IsA(UMaterialInterface::StaticClass());
 	if ((!bIsObjectAMaterial && ImportType == EInterchangePipelineContext::AssetReimport)
@@ -507,7 +507,7 @@ void UInterchangeGenericMaterialPipeline::AdjustSettingsForContext(EInterchangeP
 			HidePropertiesOfCategory(OuterMostPipeline, this, HideCategoryName);
 		}
 	}
-
+#endif //WITH_EDITOR
 	using namespace UE::Interchange;
 
 	if (!InterchangeGenericMaterialPipeline::Private::AreRequiredPackagesLoaded())
@@ -515,6 +515,49 @@ void UInterchangeGenericMaterialPipeline::AdjustSettingsForContext(EInterchangeP
 		UE_LOG(LogInterchangePipeline, Warning, TEXT("UInterchangeGenericMaterialPipeline: Some required packages are missing. Material import might be wrong"));
 	}
 }
+#if WITH_EDITOR
+
+void UInterchangeGenericMaterialPipeline::FilterPropertiesFromTranslatedData(UInterchangeBaseNodeContainer* InBaseNodeContainer)
+{
+	Super::FilterPropertiesFromTranslatedData(InBaseNodeContainer);
+
+	//Filter all material pipeline properties if there is no translated material.
+	TArray<FString> TmpMaterialNodes;
+	InBaseNodeContainer->GetNodes(UInterchangeShaderGraphNode::StaticClass(), TmpMaterialNodes);
+	uint32 MaterialCount = TmpMaterialNodes.Num();
+	InBaseNodeContainer->GetNodes(UInterchangeMaterialInstanceNode::StaticClass(), TmpMaterialNodes);
+	MaterialCount += TmpMaterialNodes.Num();
+	if(MaterialCount == 0)
+	{
+		bImportMaterials = false;
+		TArray<FString> HideCategories;
+		//Filter out all material properties
+		HideCategories.Add(TEXT("Materials"));
+		if (UInterchangePipelineBase* OuterMostPipeline = GetMostPipelineOuter())
+		{
+			for (const FString& HideCategoryName : HideCategories)
+			{
+				HidePropertiesOfCategory(OuterMostPipeline, this, HideCategoryName);
+			}
+		}
+	}
+
+	if(TexturePipeline)
+	{
+		TexturePipeline->FilterPropertiesFromTranslatedData(InBaseNodeContainer);
+	}
+}
+
+bool UInterchangeGenericMaterialPipeline::IsPropertyChangeNeedRefresh(const FPropertyChangedEvent& PropertyChangedEvent)
+{
+	if (TexturePipeline && TexturePipeline->IsPropertyChangeNeedRefresh(PropertyChangedEvent))
+	{
+		return true;
+	}
+	return Super::IsPropertyChangeNeedRefresh(PropertyChangedEvent);
+}
+
+#endif //WITH_EDITOR
 
 void UInterchangeGenericMaterialPipeline::ExecutePipeline(UInterchangeBaseNodeContainer* InBaseNodeContainer, const TArray<UInterchangeSourceData*>& InSourceDatas)
 {
diff --git a/Engine/Plugins/Interchange/Runtime/Source/Pipelines/Private/InterchangeGenericMeshPipeline.cpp b/Engine/Plugins/Interchange/Runtime/Source/Pipelines/Private/InterchangeGenericMeshPipeline.cpp
--- a/Engine/Plugins/Interchange/Runtime/Source/Pipelines/Private/InterchangeGenericMeshPipeline.cpp
+++ b/Engine/Plugins/Interchange/Runtime/Source/Pipelines/Private/InterchangeGenericMeshPipeline.cpp
@@ -19,10 +19,11 @@
 #include UE_INLINE_GENERATED_CPP_BY_NAME(InterchangeGenericMeshPipeline)
 
 void UInterchangeGenericMeshPipeline::AdjustSettingsForContext(EInterchangePipelineContext ImportType, TObjectPtr<UObject> ReimportAsset)
-
 {
 	Super::AdjustSettingsForContext(ImportType, ReimportAsset);
 
+#if WITH_EDITOR
+
 	check(CommonSkeletalMeshesAndAnimationsProperties.IsValid());
 	if (ImportType == EInterchangePipelineContext::None)
 	{
@@ -101,14 +102,29 @@ void UInterchangeGenericMeshPipeline::AdjustSettingsForContext(EInterchangePipel
 			HidePropertiesOfCategory(OuterMostPipeline, this, HideCategoryName, bDoTransientSubPipeline);
 		}
 	}
+#endif //WITH_EDITOR
+}
+
+#if WITH_EDITOR
+
+bool UInterchangeGenericMeshPipeline::IsPropertyChangeNeedRefresh(const FPropertyChangedEvent& PropertyChangedEvent)
+{
+	if (PropertyChangedEvent.GetPropertyName() == GET_MEMBER_NAME_CHECKED(UInterchangeGenericMeshPipeline, SkeletalMeshImportContentType))
+	{
+		return true;
+	}
+	return Super::IsPropertyChangeNeedRefresh(PropertyChangedEvent);
 }
 
+#endif //WITH_EDITOR
+
 void UInterchangeGenericMeshPipeline::PreDialogCleanup(const FName PipelineStackName)
 {
 	PhysicsAsset = nullptr;
 }
 
 #if WITH_EDITOR
+
 bool UInterchangeGenericMeshPipeline::GetPropertyPossibleValues(const FName PropertyPath, TArray<FString>& PossibleValues)
 {
 	FString PropertyPathString = PropertyPath.ToString();
@@ -130,8 +146,24 @@ bool UInterchangeGenericMeshPipeline::GetPropertyPossibleValues(const FName Prop
 	//If we did not find any property call the super implementation
 	return Super::GetPropertyPossibleValues(PropertyPath, PossibleValues);
 }
+
 #endif
 
+UInterchangePipelineMeshesUtilities* UInterchangeGenericMeshPipeline::CreateMeshPipelineUtilities(UInterchangeBaseNodeContainer* InBaseNodeContainer, const UInterchangeGenericMeshPipeline* Pipeline)
+{
+	UInterchangePipelineMeshesUtilities* CreatedPipelineMeshesUtilities = UInterchangePipelineMeshesUtilities::CreateInterchangePipelineMeshesUtilities(InBaseNodeContainer);
+
+	//Set the context option to use when querying the pipeline mesh utilities
+	FInterchangePipelineMeshesUtilitiesContext DataContext;
+	DataContext.bConvertStaticMeshToSkeletalMesh = (Pipeline->CommonMeshesProperties->ForceAllMeshAsType == EInterchangeForceMeshType::IFMT_SkeletalMesh);
+	DataContext.bConvertSkeletalMeshToStaticMesh = (Pipeline->CommonMeshesProperties->ForceAllMeshAsType == EInterchangeForceMeshType::IFMT_StaticMesh);
+	DataContext.bConvertStaticsWithMorphTargetsToSkeletals = Pipeline->CommonSkeletalMeshesAndAnimationsProperties->bConvertStaticsWithMorphTargetsToSkeletals;
+	DataContext.bImportMeshesInBoneHierarchy = Pipeline->CommonSkeletalMeshesAndAnimationsProperties->bImportMeshesInBoneHierarchy;
+	DataContext.bQueryGeometryOnlyIfNoInstance = Pipeline->CommonMeshesProperties->bBakeMeshes;
+	CreatedPipelineMeshesUtilities->SetContext(DataContext);
+	return CreatedPipelineMeshesUtilities;
+}
+
 void UInterchangeGenericMeshPipeline::ExecutePipeline(UInterchangeBaseNodeContainer* InBaseNodeContainer, const TArray<UInterchangeSourceData*>& InSourceDatas)
 {
 	if (!InBaseNodeContainer)
@@ -146,16 +178,7 @@ void UInterchangeGenericMeshPipeline::ExecutePipeline(UInterchangeBaseNodeContai
 	{
 		SourceDatas.Add(SourceData);
 	}
-	PipelineMeshesUtilities = UInterchangePipelineMeshesUtilities::CreateInterchangePipelineMeshesUtilities(BaseNodeContainer);
-
-	//Set the context option to use when querying the pipeline mesh utilities
-	FInterchangePipelineMeshesUtilitiesContext DataContext;
-	DataContext.bConvertStaticMeshToSkeletalMesh = (CommonMeshesProperties->ForceAllMeshAsType == EInterchangeForceMeshType::IFMT_SkeletalMesh);
-	DataContext.bConvertSkeletalMeshToStaticMesh = (CommonMeshesProperties->ForceAllMeshAsType == EInterchangeForceMeshType::IFMT_StaticMesh);
-	DataContext.bConvertStaticsWithMorphTargetsToSkeletals = CommonSkeletalMeshesAndAnimationsProperties->bConvertStaticsWithMorphTargetsToSkeletals;
-	DataContext.bImportMeshesInBoneHierarchy = CommonSkeletalMeshesAndAnimationsProperties->bImportMeshesInBoneHierarchy;
-	DataContext.bQueryGeometryOnlyIfNoInstance = CommonMeshesProperties->bBakeMeshes;
-	PipelineMeshesUtilities->SetContext(DataContext);
+	PipelineMeshesUtilities = CreateMeshPipelineUtilities(BaseNodeContainer, this);
 
 	//Create skeletalmesh factory nodes
 	ExecutePreImportPipelineSkeletalMesh();
diff --git a/Engine/Plugins/Interchange/Runtime/Source/Pipelines/Private/InterchangeGenericTexturePipeline.cpp b/Engine/Plugins/Interchange/Runtime/Source/Pipelines/Private/InterchangeGenericTexturePipeline.cpp
--- a/Engine/Plugins/Interchange/Runtime/Source/Pipelines/Private/InterchangeGenericTexturePipeline.cpp
+++ b/Engine/Plugins/Interchange/Runtime/Source/Pipelines/Private/InterchangeGenericTexturePipeline.cpp
@@ -113,7 +113,7 @@ namespace UE::Interchange::Private
 void UInterchangeGenericTexturePipeline::AdjustSettingsForContext(EInterchangePipelineContext ImportType, TObjectPtr<UObject> ReimportAsset)
 {
 	Super::AdjustSettingsForContext(ImportType, ReimportAsset);
-
+#if WITH_EDITOR
 	TArray<FString> HideCategories;
 	bool bIsObjectATexture = !ReimportAsset ? false : ReimportAsset.IsA(UTexture::StaticClass());
 	if( (!bIsObjectATexture && ImportType == EInterchangePipelineContext::AssetReimport)
@@ -132,8 +132,31 @@ void UInterchangeGenericTexturePipeline::AdjustSettingsForContext(EInterchangePi
 			HidePropertiesOfCategory(OuterMostPipeline, this, HideCategoryName);
 		}
 	}
+#endif //WITH_EDITOR
+}
+
+#if WITH_EDITOR
+
+void UInterchangeGenericTexturePipeline::FilterPropertiesFromTranslatedData(UInterchangeBaseNodeContainer* InBaseNodeContainer)
+{
+	Super::FilterPropertiesFromTranslatedData(InBaseNodeContainer);
+
+	//Filter all material pipeline properties if there is no translated material.
+	TArray<FString> TmpTextureNodes;
+	InBaseNodeContainer->GetNodes(UInterchangeTextureNode::StaticClass(), TmpTextureNodes);
+	if (TmpTextureNodes.Num() == 0)
+	{
+		bImportTextures = false;
+		//Filter out all Textures properties
+		if (UInterchangePipelineBase* OuterMostPipeline = GetMostPipelineOuter())
+		{
+			HidePropertiesOfCategory(OuterMostPipeline, this, TEXT("Textures"));
+		}
+	}
 }
 
+#endif //WITH_EDITOR
+
 void UInterchangeGenericTexturePipeline::ExecutePipeline(UInterchangeBaseNodeContainer* InBaseNodeContainer, const TArray<UInterchangeSourceData*>& InSourceDatas)
 {
 	if (!InBaseNodeContainer)
diff --git a/Engine/Plugins/Interchange/Runtime/Source/Pipelines/Public/InterchangeGenericAssetsPipeline.h b/Engine/Plugins/Interchange/Runtime/Source/Pipelines/Public/InterchangeGenericAssetsPipeline.h
--- a/Engine/Plugins/Interchange/Runtime/Source/Pipelines/Public/InterchangeGenericAssetsPipeline.h
+++ b/Engine/Plugins/Interchange/Runtime/Source/Pipelines/Public/InterchangeGenericAssetsPipeline.h
@@ -82,7 +82,11 @@ public:
 
 
 	virtual void AdjustSettingsForContext(EInterchangePipelineContext ImportType, TObjectPtr<UObject> ReimportAsset) override;
+#if WITH_EDITOR
+	virtual void FilterPropertiesFromTranslatedData(UInterchangeBaseNodeContainer* InBaseNodeContainer) override;
 
+	virtual bool IsPropertyChangeNeedRefresh(const FPropertyChangedEvent& PropertyChangedEvent) override;
+#endif //WITH_EDITOR
 	virtual bool IsScripted() override
 	{
 		return false;
diff --git a/Engine/Plugins/Interchange/Runtime/Source/Pipelines/Public/InterchangeGenericAssetsPipelineSharedSettings.h b/Engine/Plugins/Interchange/Runtime/Source/Pipelines/Public/InterchangeGenericAssetsPipelineSharedSettings.h
--- a/Engine/Plugins/Interchange/Runtime/Source/Pipelines/Public/InterchangeGenericAssetsPipelineSharedSettings.h
+++ b/Engine/Plugins/Interchange/Runtime/Source/Pipelines/Public/InterchangeGenericAssetsPipelineSharedSettings.h
@@ -104,6 +104,20 @@ public:
 	/** If true, degenerate triangles will be removed. */
 	UPROPERTY(EditAnywhere, BlueprintReadWrite, Category = "Common Meshes", meta = (SubCategory = "Build"))
 	bool bRemoveDegenerates = false;
+#if WITH_EDITOR
+	virtual bool IsPropertyChangeNeedRefresh(const FPropertyChangedEvent& PropertyChangedEvent) override
+	{
+		if (PropertyChangedEvent.Property->GetName() == GET_MEMBER_NAME_CHECKED(UInterchangeGenericCommonMeshesProperties, ForceAllMeshAsType))
+		{
+			return true;
+		}
+		if (PropertyChangedEvent.Property->GetName() == GET_MEMBER_NAME_CHECKED(UInterchangeGenericCommonMeshesProperties, bBakeMeshes))
+		{
+			return true;
+		}
+		return Super::IsPropertyChangeNeedRefresh(PropertyChangedEvent);
+	}
+#endif //WITH_EDITOR
 };
 
 UCLASS(BlueprintType, hidedropdown, Experimental)
@@ -130,7 +144,20 @@ public:
 	/** Allow to convert static mesh using morph target to skeletal mesh" */
 	UPROPERTY(EditAnywhere, BlueprintReadWrite, Category = "Static Meshes")
 	bool bConvertStaticsWithMorphTargetsToSkeletals = false;
-
+#if WITH_EDITOR
+	virtual bool IsPropertyChangeNeedRefresh(const FPropertyChangedEvent& PropertyChangedEvent) override
+	{
+		if (PropertyChangedEvent.Property->GetName() == GET_MEMBER_NAME_CHECKED(UInterchangeGenericCommonSkeletalMeshesAndAnimationsProperties, bConvertStaticsWithMorphTargetsToSkeletals))
+		{
+			return true;
+		}
+		else if (PropertyChangedEvent.Property->GetName() == GET_MEMBER_NAME_CHECKED(UInterchangeGenericCommonSkeletalMeshesAndAnimationsProperties, bImportMeshesInBoneHierarchy))
+		{
+			return true;
+		}
+		return Super::IsPropertyChangeNeedRefresh(PropertyChangedEvent);
+	}
+#endif //WITH_EDITOR
 	virtual bool IsSettingsAreValid(TOptional<FText>& OutInvalidReason) const override
 	{
 		if (bImportOnlyAnimations && !Skeleton.IsValid())
diff --git a/Engine/Plugins/Interchange/Runtime/Source/Pipelines/Public/InterchangeGenericMaterialPipeline.h b/Engine/Plugins/Interchange/Runtime/Source/Pipelines/Public/InterchangeGenericMaterialPipeline.h
--- a/Engine/Plugins/Interchange/Runtime/Source/Pipelines/Public/InterchangeGenericMaterialPipeline.h
+++ b/Engine/Plugins/Interchange/Runtime/Source/Pipelines/Public/InterchangeGenericMaterialPipeline.h
@@ -74,6 +74,10 @@ public:
 	virtual void PreDialogCleanup(const FName PipelineStackName) override;
 	virtual bool IsSettingsAreValid(TOptional<FText>& OutInvalidReason) const override;
 	virtual void AdjustSettingsForContext(EInterchangePipelineContext ImportType, TObjectPtr<UObject> ReimportAsset) override;
+#if WITH_EDITOR
+	virtual void FilterPropertiesFromTranslatedData(UInterchangeBaseNodeContainer* InBaseNodeContainer) override;
+	virtual bool IsPropertyChangeNeedRefresh(const FPropertyChangedEvent& PropertyChangedEvent) override;
+#endif //WITH_EDITOR
 
 protected:
 	virtual void ExecutePipeline(UInterchangeBaseNodeContainer* InBaseNodeContainer, const TArray<UInterchangeSourceData*>& InSourceDatas) override;
diff --git a/Engine/Plugins/Interchange/Runtime/Source/Pipelines/Public/InterchangeGenericMeshPipeline.h b/Engine/Plugins/Interchange/Runtime/Source/Pipelines/Public/InterchangeGenericMeshPipeline.h
--- a/Engine/Plugins/Interchange/Runtime/Source/Pipelines/Public/InterchangeGenericMeshPipeline.h
+++ b/Engine/Plugins/Interchange/Runtime/Source/Pipelines/Public/InterchangeGenericMeshPipeline.h
@@ -197,9 +197,12 @@ public:
 	virtual void PreDialogCleanup(const FName PipelineStackName) override;
 
 #if WITH_EDITOR
+	virtual bool IsPropertyChangeNeedRefresh(const FPropertyChangedEvent& PropertyChangedEvent) override;
 	virtual bool GetPropertyPossibleValues(const FName PropertyPath, TArray<FString>& PossibleValues) override;
 #endif
 
+	static UInterchangePipelineMeshesUtilities* CreateMeshPipelineUtilities(UInterchangeBaseNodeContainer* InBaseNodeContainer, const UInterchangeGenericMeshPipeline* Pipeline);
+
 protected:
 	virtual void ExecutePipeline(UInterchangeBaseNodeContainer* InBaseNodeContainer, const TArray<UInterchangeSourceData*>& InSourceDatas) override;
 
diff --git a/Engine/Plugins/Interchange/Runtime/Source/Pipelines/Public/InterchangeGenericTexturePipeline.h b/Engine/Plugins/Interchange/Runtime/Source/Pipelines/Public/InterchangeGenericTexturePipeline.h
--- a/Engine/Plugins/Interchange/Runtime/Source/Pipelines/Public/InterchangeGenericTexturePipeline.h
+++ b/Engine/Plugins/Interchange/Runtime/Source/Pipelines/Public/InterchangeGenericTexturePipeline.h
@@ -68,6 +68,12 @@ public:
 	virtual void ExecutePipeline(UInterchangeBaseNodeContainer* InBaseNodeContainer, const TArray<UInterchangeSourceData*>& InSourceDatas) override;
 	virtual void ExecutePostFactoryPipeline(const UInterchangeBaseNodeContainer* BaseNodeContainer, const FString& NodeKey, UObject* CreatedAsset, bool bIsAReimport) override;
 
+#if WITH_EDITOR
+
+	virtual void FilterPropertiesFromTranslatedData(UInterchangeBaseNodeContainer* InBaseNodeContainer) override;
+
+#endif //WITH_EDITOR
+
 protected:
 	UInterchangeTextureFactoryNode* HandleCreationOfTextureFactoryNode(const UInterchangeTextureNode* TextureNode);
 
diff --git a/Engine/Source/Runtime/Interchange/Core/Private/InterchangePipelineBase.cpp b/Engine/Source/Runtime/Interchange/Core/Private/InterchangePipelineBase.cpp
--- a/Engine/Source/Runtime/Interchange/Core/Private/InterchangePipelineBase.cpp
+++ b/Engine/Source/Runtime/Interchange/Core/Private/InterchangePipelineBase.cpp
@@ -311,9 +311,10 @@ UInterchangePipelineBase* UInterchangePipelineBase::GetMostPipelineOuter() const
 	return Top;
 }
 
+#if WITH_EDITOR
+
 void UInterchangePipelineBase::InternalToggleVisibilityPropertiesOfMetaDataValue(UInterchangePipelineBase* OuterMostPipeline, UInterchangePipelineBase* Pipeline, bool bDoTransientSubPipeline, const FString& MetaDataKey, const FString& MetaDataValue, const bool VisibilityState)
 {
-#if WITH_EDITOR
 	UClass* PipelineClass = Pipeline->GetClass();
 	for (FProperty* Property = PipelineClass->PropertyLink; Property; Property = Property->PropertyLinkNext)
 	{
@@ -339,10 +340,10 @@ void UInterchangePipelineBase::InternalToggleVisibilityPropertiesOfMetaDataValue
 			if (CategoryName.Equals(MetaDataValue))
 			{
 				OuterMostPipeline->FindOrAddPropertyStates(PropertyPath).ReimportStates.bVisible = VisibilityState;
+				OuterMostPipeline->FindOrAddPropertyStates(PropertyPath).ImportStates.bVisible = VisibilityState;
 			}
 		}
 	}
-#endif
 }
 
 void UInterchangePipelineBase::HidePropertiesOfCategory(UInterchangePipelineBase* OuterMostPipeline, UInterchangePipelineBase* Pipeline, const FString& HideCategoryName, bool bDoTransientSubPipeline /*= false*/)
@@ -357,6 +358,32 @@ void UInterchangePipelineBase::HidePropertiesOfSubCategory(UInterchangePipelineB
 	InternalToggleVisibilityPropertiesOfMetaDataValue(OuterMostPipeline, Pipeline, bDoTransientSubPipeline, TEXT("SubCategory"), HideSubCategoryName, bVisibilityState);
 }
 
+void UInterchangePipelineBase::HideProperty(UInterchangePipelineBase* OuterMostPipeline, UInterchangePipelineBase* Pipeline, const FName& HidePropertyName)
+{
+	constexpr bool bVisibilityState = false;
+	UClass* PipelineClass = Pipeline->GetClass();
+	for (FProperty* Property = PipelineClass->PropertyLink; Property; Property = Property->PropertyLinkNext)
+	{
+		FObjectProperty* SubObject = CastField<FObjectProperty>(Property);
+		//Do not save a transient property
+		if (SubObject || Property->HasAnyPropertyFlags(CPF_Transient))
+		{
+			continue;
+		}
+
+		const FName PropertyName = Property->GetFName();
+		const FName PropertyPath = FName(Property->GetPathName());
+		if (HidePropertyName != PropertyName)
+		{
+			continue;
+		}
+		OuterMostPipeline->FindOrAddPropertyStates(PropertyPath).ImportStates.bVisible = bVisibilityState;
+		OuterMostPipeline->FindOrAddPropertyStates(PropertyPath).ReimportStates.bVisible = bVisibilityState;
+	}
+}
+
+#endif //WITH_EDITOR
+
 struct FWeakObjectPtrData
 {
 	FString PropertyName;
diff --git a/Engine/Source/Runtime/Interchange/Core/Public/InterchangePipelineBase.h b/Engine/Source/Runtime/Interchange/Core/Public/InterchangePipelineBase.h
--- a/Engine/Source/Runtime/Interchange/Core/Public/InterchangePipelineBase.h
+++ b/Engine/Source/Runtime/Interchange/Core/Public/InterchangePipelineBase.h
@@ -108,19 +108,6 @@ class UInterchangePipelineBase : public UObject
 
 public:
 
-	UE_DEPRECATED(5.2, "This function is replace by ScriptedExecutePipeline.")
-	UFUNCTION(BlueprintNativeEvent, BlueprintCallable, Category = "Interchange | Pipeline")
-	INTERCHANGECORE_API void ScriptedExecutePreImportPipeline(UInterchangeBaseNodeContainer* BaseNodeContainer, const TArray<UInterchangeSourceData*>& SourceDatas);
-
-	UE_DEPRECATED(5.2, "This internal public function is deprecated and not replace by any.")
-	void ScriptedExecutePreImportPipeline_Implementation(UInterchangeBaseNodeContainer* BaseNodeContainer, const TArray<UInterchangeSourceData*>& SourceDatas)
-	{
-		PRAGMA_DISABLE_DEPRECATION_WARNINGS
-		//By default we call the virtual import pipeline execution
-		ExecutePreImportPipeline(BaseNodeContainer, SourceDatas);
-		PRAGMA_ENABLE_DEPRECATION_WARNINGS
-	}
-
 	/**
 	 * ScriptedExecutePipeline, is call after the translation and before we parse the graph to call the factory.
 	 * This is where factory node should be created by the pipeline.
@@ -176,17 +163,21 @@ public:
 		ExecuteExportPipeline(BaseNodeContainer);
 	}
 
+	UE_DEPRECATED(5.4, "This function will be remove, call CanExecuteOnAnyThread function")
+	INTERCHANGECORE_API bool ScriptedCanExecuteOnAnyThread(EInterchangePipelineTask PipelineTask)
+	{
+		return CanExecuteOnAnyThread(PipelineTask);
+	}
+
 	/**
-	 * Non virtual helper to allow blueprint to implement event base function let the interchange know if it can run asynchronously.
-	 * the Interchange manager is calling this function not the virtual one that is call by the default implementation.
+	 * This function tell the interchange manager if we can execute this pipeline in async mode. If it return false, the ScriptedExecuteImportPipeline
+	 * will be call on the main thread (GameThread), if true it will be run in a background thread and possibly in parallel. If there is multiple
+	 * import process in same time.
+	 *
 	 */
-	UFUNCTION(BlueprintNativeEvent, BlueprintCallable, Category = "Interchange | Pipeline")
-	INTERCHANGECORE_API bool ScriptedCanExecuteOnAnyThread(EInterchangePipelineTask PipelineTask);
-
-	/** The default implementation (call if the blueprint do not have any implementation) will call the virtual CanExecuteAsync */
-	bool ScriptedCanExecuteOnAnyThread_Implementation(EInterchangePipelineTask PipelineTask)
+	INTERCHANGECORE_API virtual bool CanExecuteOnAnyThread(EInterchangePipelineTask PipelineTask)
 	{
-		return CanExecuteOnAnyThread(PipelineTask);
+		return true;
 	}
 
 	/**
@@ -218,6 +209,20 @@ public:
 
 	INTERCHANGECORE_API void SaveSettings(const FName PipelineStackName);
 
+	/**
+	 * This function is call before we show the pipeline dialog. Pipeline that override it can change the existing settings according to the re-import type.
+	 * The function is also call when we import or re-import custom LOD and alternate skinning.
+	 *
+	 * @Note - The function will set the context of the pipeline
+	 * @Param ReimportType - Tell pipeline what re-import type the user want to achieve.
+	 * @Param ReimportAsset - This is an optional parameter which is set when re-importing an asset.
+	 */
+	INTERCHANGECORE_API virtual void AdjustSettingsForContext(EInterchangePipelineContext ReimportType, TObjectPtr<UObject> ReimportAsset);
+	INTERCHANGECORE_API virtual void AdjustSettingsFromCache();
+
+	/** Transfer the source pipeline adjust settings to this pipeline. */
+	INTERCHANGECORE_API void TransferAdjustSettings(UInterchangePipelineBase* SourcePipeline);
+
 	/**
 	 * This function is called before showing the import dialog it is not called doing a re-import.
 	 */
@@ -232,19 +237,25 @@ public:
 		return true;
 	}
 
+#if WITH_EDITOR
 	/**
-	 * This function is call when before we show the pipeline dialog. Pipeline that override it can change the existing settings according to the re-import type.
-	 * The function is also call when we import or re-import custom LOD and alternate skinning.
-	 * 
-	 * @Note - The function will set the context of the pipeline
-	 * @Param ReimportType - Tell pipeline what re-import type the user want to achieve.
-	 * @Param ReimportAsset - This is an optional parameter which is set when re-importing an asset.
+	 * Filter the pipeline properties from the translated data. This function is call by the import dialog after having duplicate and load the settings of the pipeline.
+	 * If some specific options change, the UI must refresh the filter, see function IsPropertyChangeNeedRefresh.
 	 */
-	INTERCHANGECORE_API virtual void AdjustSettingsForContext(EInterchangePipelineContext ReimportType, TObjectPtr<UObject> ReimportAsset);
-	INTERCHANGECORE_API virtual void AdjustSettingsFromCache();
+	INTERCHANGECORE_API virtual void FilterPropertiesFromTranslatedData(UInterchangeBaseNodeContainer* InBaseNodeContainer)
+	{
+		//The base class function do not have anything to filter
+	}
 
-	/** Transfer the source pipeline adjust settings to this pipeline. */
-	INTERCHANGECORE_API void TransferAdjustSettings(UInterchangePipelineBase* SourcePipeline);
+	/**
+	 * The import dialog will call this function when the user change a specific property, return true if the pipeline UI should be refresh.
+	 * A refresh will call the function FilterPropertiesFromTranslatedData.
+	 */
+	INTERCHANGECORE_API virtual bool IsPropertyChangeNeedRefresh(const FPropertyChangedEvent& PropertyChangedEvent)
+	{
+		return false;
+	}
+#endif //WITH_EDITOR
 
 	/**
 	 * This function is used to add the given message object directly into the results for this operation.
@@ -380,17 +391,6 @@ protected:
 	{
 	}
 
-	/**
-	 * This function tell the interchange manager if we can execute this pipeline in async mode. If it return false, the ScriptedExecuteImportPipeline
-	 * will be call on the main thread (GameThread), if true it will be run in a background thread and possibly in parallel. If there is multiple
-	 * import process in same time.
-	 *
-	 */
-	virtual bool CanExecuteOnAnyThread(EInterchangePipelineTask PipelineTask)
-	{
-		return true;
-	}
-
 	virtual void SetReimportSourceIndex(UClass* ReimportObjectClass, const int32 SourceFileIndex)
 	{
 	}
@@ -406,9 +406,12 @@ protected:
 
 	INTERCHANGECORE_API UInterchangePipelineBase* GetMostPipelineOuter() const;
 
+#if WITH_EDITOR
 	static INTERCHANGECORE_API void InternalToggleVisibilityPropertiesOfMetaDataValue(UInterchangePipelineBase* OuterMostPipeline, UInterchangePipelineBase* Pipeline, bool bDoTransientSubPipeline, const FString& MetaDataKey, const FString& MetaDataValue, const bool bVisibilityState);
 	static INTERCHANGECORE_API void HidePropertiesOfCategory(UInterchangePipelineBase* OuterMostPipeline, UInterchangePipelineBase* Pipeline, const FString& HideCategoryName, bool bDoTransientSubPipeline = false);
 	static INTERCHANGECORE_API void HidePropertiesOfSubCategory(UInterchangePipelineBase* OuterMostPipeline, UInterchangePipelineBase* Pipeline, const FString& HideSubCategoryName, bool bDoTransientSubPipeline = false);
+	static INTERCHANGECORE_API void HideProperty(UInterchangePipelineBase* OuterMostPipeline, UInterchangePipelineBase* Pipeline, const FName& HidePropertyName);
+#endif //WITH_EDITOR
 	
 	/**
 	 * If true, the property editor for this pipeline instance will allow properties states edition.
diff --git a/Engine/Source/Runtime/Interchange/Engine/Private/InterchangeManager.cpp b/Engine/Source/Runtime/Interchange/Engine/Private/InterchangeManager.cpp
--- a/Engine/Source/Runtime/Interchange/Engine/Private/InterchangeManager.cpp
+++ b/Engine/Source/Runtime/Interchange/Engine/Private/InterchangeManager.cpp
@@ -1228,14 +1228,16 @@ void UInterchangeManager::StartQueuedTasks(bool bCancelAllTasks /*= false*/)
 
 			//Create/Start import tasks
 			FGraphEventArray PipelinePrerequistes;
-			check(QueuedTaskData.AsyncHelper->Translators.Num() == QueuedTaskData.AsyncHelper->SourceDatas.Num());
-			for (int32 SourceDataIndex = 0; SourceDataIndex < QueuedTaskData.AsyncHelper->SourceDatas.Num(); ++SourceDataIndex)
+			if(QueuedTaskData.AsyncHelper->TranslatorTasks.Num() == 0)
 			{
-				//Log the source we begin importing
-				UE_LOG(LogInterchangeEngine, Display, TEXT("Interchange start importing source [%s]"), *QueuedTaskData.AsyncHelper->SourceDatas[SourceDataIndex]->ToDisplayString());
-
-				int32 TranslatorTaskIndex = QueuedTaskData.AsyncHelper->TranslatorTasks.Add(TGraphTask<UE::Interchange::FTaskTranslator>::CreateTask().ConstructAndDispatchWhenReady(SourceDataIndex, WeakAsyncHelper));
-				PipelinePrerequistes.Add(QueuedTaskData.AsyncHelper->TranslatorTasks[TranslatorTaskIndex]);
+				check(QueuedTaskData.AsyncHelper->Translators.Num() == QueuedTaskData.AsyncHelper->SourceDatas.Num());
+				for (int32 SourceDataIndex = 0; SourceDataIndex < QueuedTaskData.AsyncHelper->SourceDatas.Num(); ++SourceDataIndex)
+				{
+					//Log the source we begin importing
+					UE_LOG(LogInterchangeEngine, Display, TEXT("Interchange start importing source [%s]"), *QueuedTaskData.AsyncHelper->SourceDatas[SourceDataIndex]->ToDisplayString());
+					int32 TranslatorTaskIndex = QueuedTaskData.AsyncHelper->TranslatorTasks.Add(TGraphTask<UE::Interchange::FTaskTranslator>::CreateTask().ConstructAndDispatchWhenReady(SourceDataIndex, WeakAsyncHelper));
+					PipelinePrerequistes.Add(QueuedTaskData.AsyncHelper->TranslatorTasks[TranslatorTaskIndex]);
+				}
 			}
 
 			FGraphEventArray GraphParsingPrerequistes;
@@ -1425,7 +1427,7 @@ UInterchangeManager::ImportInternal(const FString& ContentPath, const UInterchan
 	UInterchangeSourceData* DuplicateSourceData = Cast<UInterchangeSourceData>(StaticDuplicateObject(SourceData, GetTransientPackage()));
 	//Array of source data to build one graph per source
 	AsyncHelper->SourceDatas.Add(DuplicateSourceData);
-
+	constexpr int32 SourceIndex = 0;
 	//Get all the translators for the source datas
 	for (int32 SourceDataIndex = 0; SourceDataIndex < AsyncHelper->SourceDatas.Num(); ++SourceDataIndex)
 	{
@@ -1497,6 +1499,24 @@ UInterchangeManager::ImportInternal(const FString& ContentPath, const UInterchan
 		const bool bShowPipelineStacksConfigurationDialog = false;
 #endif
 
+		auto TranslateSourceFile = [&AsyncHelper]()
+		{
+			FScopedSlowTask Progress(2.f, NSLOCTEXT("InterchangeManager", "TranslatingSourceFile...", "Translating source file..."));
+			Progress.MakeDialog();
+			Progress.EnterProgressFrame(1.f);
+			//Translate the source
+			FGraphEventArray PipelinePrerequistes;
+			check(AsyncHelper->Translators.Num() == AsyncHelper->SourceDatas.Num());
+			for (int32 SourceDataIndex = 0; SourceDataIndex < AsyncHelper->SourceDatas.Num(); ++SourceDataIndex)
+			{
+				//Log the source we begin importing
+				UE_LOG(LogInterchangeEngine, Display, TEXT("Interchange start importing source [%s]"), *AsyncHelper->SourceDatas[SourceDataIndex]->ToDisplayString());
+				int32 TranslatorTaskIndex = AsyncHelper->TranslatorTasks.Add(TGraphTask<UE::Interchange::FTaskTranslator>::CreateTask(nullptr, ENamedThreads::GameThread).ConstructAndDispatchWhenReady(SourceDataIndex, AsyncHelper));
+				AsyncHelper->TranslatorTasks[TranslatorTaskIndex]->Wait();
+			}
+			Progress.EnterProgressFrame(1.f);
+		};
+
 		if (FEngineAnalytics::IsAvailable())
 		{
 			Attribs.Add(FAnalyticsEventAttribute(TEXT("ShowImportDialog"), bShowPipelineStacksConfigurationDialog));
@@ -1544,7 +1564,6 @@ UInterchangeManager::ImportInternal(const FString& ContentPath, const UInterchan
 
 		{
 			UE::Interchange::FScopedTranslator ScopedTranslator(SourceData);
-			const UInterchangeTranslatorBase* Translator = ScopedTranslator.GetTranslator();
 
 			for (const TPair<FName, FInterchangePipelineStack>& PipelineStackInfo : DefaultPipelineStacks)
 			{
@@ -1559,7 +1578,7 @@ UInterchangeManager::ImportInternal(const FString& ContentPath, const UInterchan
 				for (const FInterchangeTranslatorPipelines& TranslatorPipelines : PipelineStack.PerTranslatorPipelines)
 				{
 					const UClass* TranslatorClass = TranslatorPipelines.Translator.LoadSynchronous();
-					if (Translator->IsA(TranslatorClass))
+					if (ScopedTranslator.GetTranslator()->IsA(TranslatorClass))
 					{
 						Pipelines = &TranslatorPipelines.Pipelines;
 						break;
@@ -1587,8 +1606,10 @@ UInterchangeManager::ImportInternal(const FString& ContentPath, const UInterchan
 		{
 			if (RegisteredPipelineConfiguration && bShowPipelineStacksConfigurationDialog && !bIsUnattended)
 			{
+				TranslateSourceFile();
+				UInterchangeBaseNodeContainer* BaseNodeContainer = AsyncHelper->BaseNodeContainers[SourceIndex].Get();
 				//Show the dialog, a plugin should have registered this dialog. We use a plugin to be able to use editor code when doing UI
-				EInterchangePipelineConfigurationDialogResult DialogResult = RegisteredPipelineConfiguration->ScriptedShowReimportPipelineConfigurationDialog(PipelineStacks, OutPipelines, DuplicateSourceData);
+				EInterchangePipelineConfigurationDialogResult DialogResult = RegisteredPipelineConfiguration->ScriptedShowReimportPipelineConfigurationDialog(PipelineStacks, OutPipelines, DuplicateSourceData, BaseNodeContainer);
 				if (DialogResult == EInterchangePipelineConfigurationDialogResult::Cancel)
 				{
 					bImportCanceled = true;
@@ -1614,10 +1635,12 @@ UInterchangeManager::ImportInternal(const FString& ContentPath, const UInterchan
 		{
 			if (RegisteredPipelineConfiguration && bShowPipelineStacksConfigurationDialog)
 			{
+				TranslateSourceFile();
+				UInterchangeBaseNodeContainer* BaseNodeContainer = AsyncHelper->BaseNodeContainers[SourceIndex].Get();
 				//Show the dialog, a plugin should have register this dialog. We use a plugin to be able to use editor code when doing UI
 				EInterchangePipelineConfigurationDialogResult DialogResult = bImportScene
-					? RegisteredPipelineConfiguration->ScriptedShowScenePipelineConfigurationDialog(PipelineStacks, OutPipelines, DuplicateSourceData)
-					: RegisteredPipelineConfiguration->ScriptedShowPipelineConfigurationDialog(PipelineStacks, OutPipelines, DuplicateSourceData);
+					? RegisteredPipelineConfiguration->ScriptedShowScenePipelineConfigurationDialog(PipelineStacks, OutPipelines, DuplicateSourceData, BaseNodeContainer)
+					: RegisteredPipelineConfiguration->ScriptedShowPipelineConfigurationDialog(PipelineStacks, OutPipelines, DuplicateSourceData, BaseNodeContainer);
 
 				if (DialogResult == EInterchangePipelineConfigurationDialogResult::Cancel)
 				{
diff --git a/Engine/Source/Runtime/Interchange/Engine/Private/Tasks/InterchangeTaskPipeline.h b/Engine/Source/Runtime/Interchange/Engine/Private/Tasks/InterchangeTaskPipeline.h
--- a/Engine/Source/Runtime/Interchange/Engine/Private/Tasks/InterchangeTaskPipeline.h
+++ b/Engine/Source/Runtime/Interchange/Engine/Private/Tasks/InterchangeTaskPipeline.h
@@ -39,7 +39,7 @@ namespace UE
 					return ENamedThreads::GameThread;
 				}
 
-				return PipelineBase.Get()->ScriptedCanExecuteOnAnyThread(EInterchangePipelineTask::PostTranslator) ? ENamedThreads::AnyBackgroundThreadNormalTask : ENamedThreads::GameThread;
+				return PipelineBase.Get()->CanExecuteOnAnyThread(EInterchangePipelineTask::PostTranslator) ? ENamedThreads::AnyBackgroundThreadNormalTask : ENamedThreads::GameThread;
 			}
 
 			static ESubsequentsMode::Type GetSubsequentsMode()
@@ -90,7 +90,7 @@ namespace UE
 				}
 
 				//Ask the pipeline implementation
-				if (AsyncHelper->Pipelines[PipelineIndex]->ScriptedCanExecuteOnAnyThread(EInterchangePipelineTask::PostImport))
+				if (AsyncHelper->Pipelines[PipelineIndex]->CanExecuteOnAnyThread(EInterchangePipelineTask::PostImport))
 				{
 					return ENamedThreads::AnyBackgroundThreadNormalTask;
 				}
diff --git a/Engine/Source/Runtime/Interchange/Engine/Public/InterchangePipelineConfigurationBase.h b/Engine/Source/Runtime/Interchange/Engine/Public/InterchangePipelineConfigurationBase.h
--- a/Engine/Source/Runtime/Interchange/Engine/Public/InterchangePipelineConfigurationBase.h
+++ b/Engine/Source/Runtime/Interchange/Engine/Public/InterchangePipelineConfigurationBase.h
@@ -42,49 +42,55 @@ public:
 	UFUNCTION(BlueprintNativeEvent, BlueprintCallable, Category = "Interchange | Translator")
 	INTERCHANGEENGINE_API EInterchangePipelineConfigurationDialogResult ScriptedShowPipelineConfigurationDialog(TArray<FInterchangeStackInfo>& PipelineStacks
 		, TArray<UInterchangePipelineBase*>& OutPipelines
-		, UInterchangeSourceData* SourceData);
+		, UInterchangeSourceData* SourceData
+		, UInterchangeBaseNodeContainer* BaseNodeContainer);
 
 	/** The default implementation (call if the blueprint do not have any implementation) will call the virtual ExecuteImportPipeline */
 	EInterchangePipelineConfigurationDialogResult ScriptedShowPipelineConfigurationDialog_Implementation(TArray<FInterchangeStackInfo>& PipelineStacks
 		, TArray<UInterchangePipelineBase*>& OutPipelines
-		, UInterchangeSourceData* SourceData)
+		, UInterchangeSourceData* SourceData
+		, UInterchangeBaseNodeContainer* BaseNodeContainer)
 	{
 		//By default we call the virtual import pipeline execution
-		return ShowPipelineConfigurationDialog(PipelineStacks, OutPipelines, SourceData);
+		return ShowPipelineConfigurationDialog(PipelineStacks, OutPipelines, SourceData, BaseNodeContainer);
 	}
 
 	/**
-	 * Non virtual helper to allow blueprint to implement event base function to implement the ShowPipelineConfigurationDialog,
+	 * Non virtual helper to allow blueprint to implement event base function to implement the ShowScenePipelineConfigurationDialog,
 	 */
 	UFUNCTION(BlueprintNativeEvent, BlueprintCallable, Category = "Interchange | Translator")
 	INTERCHANGEENGINE_API EInterchangePipelineConfigurationDialogResult ScriptedShowScenePipelineConfigurationDialog(TArray<FInterchangeStackInfo>& PipelineStacks
 		, TArray<UInterchangePipelineBase*>& OutPipelines
-		, UInterchangeSourceData* SourceData);
+		, UInterchangeSourceData* SourceData
+		, UInterchangeBaseNodeContainer* BaseNodeContainer);
 
 	/** The default implementation (call if the blueprint do not have any implementation) will call the virtual ExecuteImportPipeline */
 	EInterchangePipelineConfigurationDialogResult ScriptedShowScenePipelineConfigurationDialog_Implementation(TArray<FInterchangeStackInfo>& PipelineStacks
 		, TArray<UInterchangePipelineBase*>& OutPipelines
-		, UInterchangeSourceData* SourceData)
+		, UInterchangeSourceData* SourceData
+		, UInterchangeBaseNodeContainer* BaseNodeContainer)
 	{
 		//By default we call the virtual import pipeline execution
-		return ShowScenePipelineConfigurationDialog(PipelineStacks, OutPipelines, SourceData);
+		return ShowScenePipelineConfigurationDialog(PipelineStacks, OutPipelines, SourceData, BaseNodeContainer);
 	}
 
 	/**
-	 * Non virtual helper to allow blueprint to implement event base function to implement the ShowPipelineConfigurationDialog,
+	 * Non virtual helper to allow blueprint to implement event base function to implement the ShowReimportPipelineConfigurationDialog,
 	 */
 	UFUNCTION(BlueprintNativeEvent, BlueprintCallable, Category = "Interchange | Translator")
 	INTERCHANGEENGINE_API EInterchangePipelineConfigurationDialogResult ScriptedShowReimportPipelineConfigurationDialog(TArray<FInterchangeStackInfo>& PipelineStacks
 		, TArray<UInterchangePipelineBase*>& OutPipelines
-		, UInterchangeSourceData* SourceData);
+		, UInterchangeSourceData* SourceData
+		, UInterchangeBaseNodeContainer* BaseNodeContainer);
 
 	/** The default implementation (call if the blueprint do not have any implementation) will call the virtual ExecuteImportPipeline */
 	EInterchangePipelineConfigurationDialogResult ScriptedShowReimportPipelineConfigurationDialog_Implementation(TArray<FInterchangeStackInfo>& PipelineStacks
 		, TArray<UInterchangePipelineBase*>& OutPipelines
-		, UInterchangeSourceData* SourceData)
+		, UInterchangeSourceData* SourceData
+		, UInterchangeBaseNodeContainer* BaseNodeContainer)
 	{
 		//By default we call the virtual import pipeline execution
-		return ShowReimportPipelineConfigurationDialog(PipelineStacks, OutPipelines, SourceData);
+		return ShowReimportPipelineConfigurationDialog(PipelineStacks, OutPipelines, SourceData, BaseNodeContainer);
 	}
 
 protected:
@@ -94,7 +100,8 @@ protected:
 	 */
 	virtual EInterchangePipelineConfigurationDialogResult ShowPipelineConfigurationDialog(TArray<FInterchangeStackInfo>& PipelineStacks
 		, TArray<UInterchangePipelineBase*>& OutPipelines
-		, TWeakObjectPtr<UInterchangeSourceData> SourceData)
+		, TWeakObjectPtr<UInterchangeSourceData> SourceData
+		, TWeakObjectPtr <UInterchangeBaseNodeContainer> BaseNodeContainer)
 	{ 
 		//Not implemented
 		return EInterchangePipelineConfigurationDialogResult::Cancel;
@@ -105,7 +112,8 @@ protected:
 	 */
 	virtual EInterchangePipelineConfigurationDialogResult ShowScenePipelineConfigurationDialog(TArray<FInterchangeStackInfo>& PipelineStacks
 		, TArray<UInterchangePipelineBase*>& OutPipelines
-		, TWeakObjectPtr<UInterchangeSourceData> SourceData)
+		, TWeakObjectPtr<UInterchangeSourceData> SourceData
+		, TWeakObjectPtr <UInterchangeBaseNodeContainer> BaseNodeContainer)
 	{ 
 		//Not implemented
 		return EInterchangePipelineConfigurationDialogResult::Cancel;
@@ -116,7 +124,8 @@ protected:
 	 */
 	virtual EInterchangePipelineConfigurationDialogResult ShowReimportPipelineConfigurationDialog(TArray<FInterchangeStackInfo>& PipelineStacks
 		, TArray<UInterchangePipelineBase*>& OutPipelines
-		, TWeakObjectPtr<UInterchangeSourceData> SourceData)
+		, TWeakObjectPtr<UInterchangeSourceData> SourceData
+		, TWeakObjectPtr <UInterchangeBaseNodeContainer> BaseNodeContainer)
 	{
 		//Not implemented
 		return EInterchangePipelineConfigurationDialogResult::Cancel;
