Localize sub category name for fbx legacy import dialog and interchnage import dialog.

#rnx

diff --git a/Engine/Plugins/Interchange/Editor/Source/Pipelines/Private/InterchangeEditorPipelineDetails.cpp b/Engine/Plugins/Interchange/Editor/Source/Pipelines/Private/InterchangeEditorPipelineDetails.cpp
--- a/Engine/Plugins/Interchange/Editor/Source/Pipelines/Private/InterchangeEditorPipelineDetails.cpp
+++ b/Engine/Plugins/Interchange/Editor/Source/Pipelines/Private/InterchangeEditorPipelineDetails.cpp
@@ -280,7 +280,7 @@ void FInterchangePipelineBaseDetailsCustomization::CustomizeDetails(IDetailLayou
 	for (const TPair<FName, TArray<FName>>& CategoryAndProperties : PropertiesPerCategorys)
 	{
 		//Category meta value Subgroup data
-		TMap<FName, IDetailGroup*> SubCategoryGroups;
+		TMap<FString, IDetailGroup*> SubCategoryGroups;
 
 		const FName CategoryName = CategoryAndProperties.Key;
 		
@@ -325,15 +325,38 @@ void FInterchangePipelineBaseDetailsCustomization::CustomizeDetails(IDetailLayou
 			FName PropertyPath = FName(PropertyPtr->GetPathName());
 			CachedDetailBuilder->HideProperty(PropertyHandle);
 
-			const FName SubCategoryData = FName(PropertyHandle->GetMetaData(TEXT("SubCategory")));
+			const FString SubCategoryData = PropertyHandle->GetMetaData(TEXT("SubCategory"));
 			IDetailGroup* GroupPtr = nullptr;
 			auto GetGroupPtr = [&GroupPtr, &SubCategoryData, &SubCategoryGroups, &Category]()
 			{
-				if (SubCategoryData != NAME_None)
+				if (!SubCategoryData.IsEmpty())
 				{
 					if (!SubCategoryGroups.Contains(SubCategoryData))
 					{
-						SubCategoryGroups.Add(SubCategoryData, &(Category.AddGroup(SubCategoryData, FText::FromName(SubCategoryData))));
+						//Localize sub category
+						FText LocalizeSubCategoryName = FText::FromString(SubCategoryData);
+						
+						if (SubCategoryData.Equals(TEXT("Build")))
+						{
+							LocalizeSubCategoryName = NSLOCTEXT("InterchangePipelineBaseDetails::CustomizeDetails", "SubCategory_Build", "Build");
+						}
+						else if (SubCategoryData.Equals(TEXT("Collision")))
+						{
+							LocalizeSubCategoryName = NSLOCTEXT("InterchangePipelineBaseDetails::CustomizeDetails", "SubCategory_Collision", "Collision");
+						}
+						else if (SubCategoryData.Equals(TEXT("Actors properties")))
+						{
+							LocalizeSubCategoryName = NSLOCTEXT("InterchangePipelineBaseDetails::CustomizeDetails", "SubCategory_Actors_properties", "Actors properties");
+						}
+						else if (SubCategoryData.Equals(TEXT("Reimport Actors")))
+						{
+							LocalizeSubCategoryName = NSLOCTEXT("InterchangePipelineBaseDetails::CustomizeDetails", "SubCategory_Reimport_Actors", "Reimport Actors");
+						}
+						else if (SubCategoryData.Equals(TEXT("Reimport Assets")))
+						{
+							LocalizeSubCategoryName = NSLOCTEXT("InterchangePipelineBaseDetails::CustomizeDetails", "SubCategory_Reimport_Assets", "Reimport Assets");
+						}
+						SubCategoryGroups.Add(SubCategoryData, &(Category.AddGroup(FName(SubCategoryData), LocalizeSubCategoryName)));
 					}
 					GroupPtr = SubCategoryGroups.FindChecked(SubCategoryData);
 				}
diff --git a/Engine/Source/Editor/DetailCustomizations/Private/FbxImportUIDetails.cpp b/Engine/Source/Editor/DetailCustomizations/Private/FbxImportUIDetails.cpp
--- a/Engine/Source/Editor/DetailCustomizations/Private/FbxImportUIDetails.cpp
+++ b/Engine/Source/Editor/DetailCustomizations/Private/FbxImportUIDetails.cpp
@@ -449,10 +449,16 @@ void FFbxImportUIDetails::CustomizeDetails( IDetailLayoutBuilder& DetailBuilder
 			}
 			else if (!SubCategoryData.IsEmpty())
 			{
-				TArray<TSharedPtr<IPropertyHandle> >& SubCategoryProperties = SubCategoriesProperties.FindOrAdd(SubCategoryData);
+				//Localize the sub category
+				FString SubCategoryLocalizeName = SubCategoryData;
+				if (SubCategoryData.Equals(TEXT("Thresholds")))
+				{
+					SubCategoryLocalizeName = LOCTEXT("Thresholds_subcategory_name", "Thresholds").ToString();
+				}
+				TArray<TSharedPtr<IPropertyHandle> >& SubCategoryProperties = SubCategoriesProperties.FindOrAdd(SubCategoryLocalizeName);
 				SubCategoryProperties.Add(Handle);
-				bool& SubCategoryAdvanced = SubCategoriesAdvanced.FindOrAdd(SubCategoryData);
-				FText& SubCategoryTooltip = SubCategoriesTooltip.FindOrAdd(SubCategoryData);
+				bool& SubCategoryAdvanced = SubCategoriesAdvanced.FindOrAdd(SubCategoryLocalizeName);
+				FText& SubCategoryTooltip = SubCategoriesTooltip.FindOrAdd(SubCategoryLocalizeName);
 				if (SubCategoryData.Equals(TEXT("Thresholds")))
 				{
 					SubCategoryAdvanced = true;
