Interchange: Fix a validation issue with python pipeline

#rnx

diff --git a/Engine/Plugins/Interchange/Editor/Source/Pipelines/Private/InterchangePipelineFactories.cpp b/Engine/Plugins/Interchange/Editor/Source/Pipelines/Private/InterchangePipelineFactories.cpp
--- a/Engine/Plugins/Interchange/Editor/Source/Pipelines/Private/InterchangePipelineFactories.cpp
+++ b/Engine/Plugins/Interchange/Editor/Source/Pipelines/Private/InterchangePipelineFactories.cpp
@@ -155,6 +155,8 @@ UObject* UInterchangePythonPipelineAssetFactory::FactoryCreateNew(UClass* Class,
 	if (PythonClass && PythonClass->IsChildOf(UInterchangePipelineBase::StaticClass()))
 	{
 		Pipeline = NewObject<UInterchangePythonPipelineAsset>(InParent, SupportedClass, InName, InFlags | RF_Transactional);
+		//Python pipeline are editor only package
+		Pipeline->GetPackage()->SetPackageFlags(PKG_EditorOnly);
 		Pipeline->PythonClass = PythonClass;
 		Pipeline->GeneratePipeline();
 	}
diff --git a/Engine/Source/Runtime/Interchange/Engine/Private/InterchangePythonPipelineBase.cpp b/Engine/Source/Runtime/Interchange/Engine/Private/InterchangePythonPipelineBase.cpp
--- a/Engine/Source/Runtime/Interchange/Engine/Private/InterchangePythonPipelineBase.cpp
+++ b/Engine/Source/Runtime/Interchange/Engine/Private/InterchangePythonPipelineBase.cpp
@@ -33,6 +33,7 @@ void UInterchangePythonPipelineBase::PostEditChangeProperty(struct FPropertyChan
 
 void UInterchangePythonPipelineAsset::PostLoad()
 {
+	GetPackage()->SetPackageFlags(PKG_EditorOnly);
 	Super::PostLoad();
 	GeneratePipeline();
 }
