Fix the reimport to correctly find the reimport handler when the file extension has change. It can happen if you drag and drop a file in the content browser.

#rnx

diff --git a/Engine/Plugins/Interchange/Editor/Source/InterchangeEditor/InterchangeEditor.Build.cs b/Engine/Plugins/Interchange/Editor/Source/InterchangeEditor/InterchangeEditor.Build.cs
--- a/Engine/Plugins/Interchange/Editor/Source/InterchangeEditor/InterchangeEditor.Build.cs
+++ b/Engine/Plugins/Interchange/Editor/Source/InterchangeEditor/InterchangeEditor.Build.cs
@@ -25,7 +25,10 @@ namespace UnrealBuildTool.Rules
 					"ContentBrowser",
 					"InterchangeCore",
 					"InterchangeEngine",
+					"InterchangeFactoryNodes",
 					"InterchangeImport",
+					"InterchangeNodes",
+					"InterchangePipelines",
 					"MessageLog",
 					"SlateCore",
 					"ToolMenus",
diff --git a/Engine/Plugins/Interchange/Editor/Source/InterchangeEditor/Private/InterchangeEditorModule.cpp b/Engine/Plugins/Interchange/Editor/Source/InterchangeEditor/Private/InterchangeEditorModule.cpp
--- a/Engine/Plugins/Interchange/Editor/Source/InterchangeEditor/Private/InterchangeEditorModule.cpp
+++ b/Engine/Plugins/Interchange/Editor/Source/InterchangeEditor/Private/InterchangeEditorModule.cpp
@@ -4,6 +4,7 @@
 #include "InterchangeEditorLog.h"
 
 #include "InterchangeManager.h"
+#include "InterchangeFbxAssetImportDataConverter.h"
 
 #include "Engine/Engine.h"
 #include "IMessageLogListing.h"
@@ -86,6 +87,7 @@ void FInterchangeEditorModule::StartupModule()
 
 		UInterchangeManager& InterchangeManager = UInterchangeManager::GetInterchangeManager();
 		InterchangeEditorModuleDelegate = InterchangeManager.OnBatchImportComplete.AddStatic(&InterchangeEditorModule::LogErrors);
+		InterchangeManager.RegisterImportDataConverter(UInterchangeFbxAssetImportDataConverter::StaticClass());
 
 		auto UnregisterItems = [InterchangeEditorModuleDelegate]()
 		{
diff --git a/Engine/Plugins/Interchange/Editor/Source/InterchangeEditor/Private/InterchangeFbxAssetImportDataConverter.cpp b/Engine/Plugins/Interchange/Editor/Source/InterchangeEditor/Private/InterchangeFbxAssetImportDataConverter.cpp
new file mode 100644
--- /dev/null
+++ b/Engine/Plugins/Interchange/Editor/Source/InterchangeEditor/Private/InterchangeFbxAssetImportDataConverter.cpp
@@ -0,0 +1,671 @@
+// Copyright Epic Games, Inc. All Rights Reserved. 
+#include "InterchangeFbxAssetImportDataConverter.h"
+
+#include "Animation/AnimSequence.h"
+#include "EditorFramework/AssetImportData.h"
+#include "Engine/SkeletalMesh.h"
+#include "Engine/SkinnedAssetCommon.h"
+#include "Engine/StaticMesh.h"
+#include "Factories/FbxAnimSequenceImportData.h"
+#include "Factories/FbxAssetImportData.h"
+#include "Factories/FbxMeshImportData.h"
+#include "Factories/FbxSkeletalMeshImportData.h"
+#include "Factories/FbxStaticMeshImportData.h"
+#include "InterchangeAnimSequenceFactoryNode.h"
+#include "InterchangeAssetImportData.h"
+#include "InterchangeGenericAnimationPipeline.h"
+#include "InterchangeGenericAssetsPipeline.h"
+#include "InterchangeGenericAssetsPipelineSharedSettings.h"
+#include "InterchangeGenericMeshPipeline.h"
+#include "InterchangeSkeletalMeshFactoryNode.h"
+#include "InterchangeStaticMeshFactoryNode.h"
+#include "InterchangeManager.h"
+#include "Nodes/InterchangeBaseNodeContainer.h"
+
+namespace UE::Interchange::Private
+{
+	void TransferSourceFileInformation(const UAssetImportData* SourceData, UAssetImportData* DestinationData)
+	{
+		TArray<FAssetImportInfo::FSourceFile> SourceFiles = SourceData->GetSourceData().SourceFiles;
+		DestinationData->SetSourceFiles(MoveTemp(SourceFiles));
+	}
+
+	void FillFbxAssetImportData(const UInterchangeGenericAssetsPipeline* GenericAssetPipeline, UFbxAssetImportData* AssetImportData)
+	{
+		AssetImportData->bConvertScene = true;
+		AssetImportData->bConvertSceneUnit = true;
+		AssetImportData->bForceFrontXAxis = false;
+		AssetImportData->bImportAsScene = false;
+		AssetImportData->ImportRotation = GenericAssetPipeline->ImportOffsetRotation;
+		AssetImportData->ImportTranslation = GenericAssetPipeline->ImportOffsetTranslation;
+		AssetImportData->ImportUniformScale = GenericAssetPipeline->ImportOffsetUniformScale;
+	}
+
+	void FillFbxMeshImportData(const UInterchangeGenericAssetsPipeline* GenericAssetPipeline, UFbxMeshImportData* MeshImportData)
+	{
+		MeshImportData->bBakePivotInVertex = false;
+		MeshImportData->bComputeWeightedNormals = GenericAssetPipeline->CommonMeshesProperties->bComputeWeightedNormals;
+		MeshImportData->bImportMeshLODs = GenericAssetPipeline->CommonMeshesProperties->bImportLods;
+		MeshImportData->bReorderMaterialToFbxOrder = true;
+		MeshImportData->bTransformVertexToAbsolute = GenericAssetPipeline->CommonMeshesProperties->bBakeMeshes;
+
+		if (GenericAssetPipeline->CommonMeshesProperties->bUseMikkTSpace)
+		{
+			MeshImportData->NormalGenerationMethod = EFBXNormalGenerationMethod::MikkTSpace;
+		}
+		else
+		{
+			MeshImportData->NormalGenerationMethod = EFBXNormalGenerationMethod::BuiltIn;
+		}
+
+		if (GenericAssetPipeline->CommonMeshesProperties->bRecomputeNormals)
+		{
+			MeshImportData->NormalImportMethod = EFBXNormalImportMethod::FBXNIM_ComputeNormals;
+		}
+		else
+		{
+			if (GenericAssetPipeline->CommonMeshesProperties->bRecomputeTangents)
+			{
+				MeshImportData->NormalImportMethod = EFBXNormalImportMethod::FBXNIM_ImportNormals;
+			}
+			else
+			{
+				MeshImportData->NormalImportMethod = EFBXNormalImportMethod::FBXNIM_ImportNormalsAndTangents;
+			}
+		}
+	}
+
+	void FillInterchangeGenericAssetsPipelineFromFbxMeshImportData(UInterchangeGenericAssetsPipeline* GenericAssetPipeline, const UFbxMeshImportData* LegacyMeshImportData)
+	{
+		GenericAssetPipeline->CommonMeshesProperties->bComputeWeightedNormals = LegacyMeshImportData->bComputeWeightedNormals;
+		GenericAssetPipeline->CommonMeshesProperties->bImportLods = LegacyMeshImportData->bImportMeshLODs;
+		GenericAssetPipeline->CommonMeshesProperties->bBakeMeshes = LegacyMeshImportData->bTransformVertexToAbsolute;
+
+		if (LegacyMeshImportData->NormalGenerationMethod == EFBXNormalGenerationMethod::MikkTSpace)
+		{
+			GenericAssetPipeline->CommonMeshesProperties->bUseMikkTSpace = true;
+		}
+		else
+		{
+			GenericAssetPipeline->CommonMeshesProperties->bUseMikkTSpace = false;
+		}
+
+		if (LegacyMeshImportData->NormalImportMethod == EFBXNormalImportMethod::FBXNIM_ComputeNormals)
+		{
+			GenericAssetPipeline->CommonMeshesProperties->bRecomputeNormals = true;
+		}
+		else
+		{
+			if (LegacyMeshImportData->NormalImportMethod == EFBXNormalImportMethod::FBXNIM_ImportNormals)
+			{
+				GenericAssetPipeline->CommonMeshesProperties->bRecomputeTangents = true;
+			}
+			else
+			{
+				GenericAssetPipeline->CommonMeshesProperties->bRecomputeTangents = false;
+			}
+		}
+	}
+
+	void FillInterchangeGenericAssetsPipelineFromFbxStaticMeshImportData(UInterchangeGenericAssetsPipeline* GenericAssetPipeline
+		, const UFbxStaticMeshImportData* StaticMeshImportData
+		, UInterchangeBaseNodeContainer* DestinationContainer
+		, const FString& NodeUniqueId
+		, const FString& NodeDisplayLabel)
+	{
+		UInterchangeStaticMeshFactoryNode* MeshNode = NewObject<UInterchangeStaticMeshFactoryNode>(DestinationContainer);
+		MeshNode->InitializeStaticMeshNode(NodeUniqueId, NodeDisplayLabel, UStaticMesh::StaticClass()->GetName());
+		DestinationContainer->AddNode(MeshNode);
+
+		FillInterchangeGenericAssetsPipelineFromFbxMeshImportData(GenericAssetPipeline, StaticMeshImportData);
+
+		GenericAssetPipeline->MeshPipeline->bImportCollision = StaticMeshImportData->bAutoGenerateCollision;
+		GenericAssetPipeline->MeshPipeline->bBuildNanite = StaticMeshImportData->bBuildNanite;
+		GenericAssetPipeline->MeshPipeline->bBuildReversedIndexBuffer = StaticMeshImportData->bBuildReversedIndexBuffer;
+		GenericAssetPipeline->MeshPipeline->bCombineStaticMeshes = StaticMeshImportData->bCombineMeshes;
+		GenericAssetPipeline->MeshPipeline->bGenerateLightmapUVs = StaticMeshImportData->bGenerateLightmapUVs;
+		GenericAssetPipeline->MeshPipeline->bOneConvexHullPerUCX = StaticMeshImportData->bOneConvexHullPerUCX;
+		GenericAssetPipeline->CommonMeshesProperties->bRemoveDegenerates = StaticMeshImportData->bRemoveDegenerates;
+		GenericAssetPipeline->MeshPipeline->DistanceFieldResolutionScale = StaticMeshImportData->DistanceFieldResolutionScale;
+		GenericAssetPipeline->MeshPipeline->LodGroup = StaticMeshImportData->StaticMeshLODGroup;
+		if (StaticMeshImportData->VertexColorImportOption == EVertexColorImportOption::Ignore)
+		{
+			GenericAssetPipeline->CommonMeshesProperties->VertexColorImportOption = EInterchangeVertexColorImportOption::IVCIO_Ignore;
+		}
+		else if (StaticMeshImportData->VertexColorImportOption == EVertexColorImportOption::Override)
+		{
+			GenericAssetPipeline->CommonMeshesProperties->VertexColorImportOption = EInterchangeVertexColorImportOption::IVCIO_Override;
+		}
+		else if (StaticMeshImportData->VertexColorImportOption == EVertexColorImportOption::Replace)
+		{
+			GenericAssetPipeline->CommonMeshesProperties->VertexColorImportOption = EInterchangeVertexColorImportOption::IVCIO_Replace;
+		}
+
+		GenericAssetPipeline->CommonMeshesProperties->VertexOverrideColor = StaticMeshImportData->VertexOverrideColor;
+	}
+
+	void FillInterchangeGenericAssetsPipelineFromFbxSkeletalMeshImportData(UInterchangeGenericAssetsPipeline* GenericAssetPipeline
+		, const UFbxSkeletalMeshImportData* SkeletalMeshImportData
+		, UInterchangeBaseNodeContainer* DestinationContainer
+		, const FString& NodeUniqueId
+		, const FString& NodeDisplayLabel)
+	{
+		UInterchangeSkeletalMeshFactoryNode* MeshNode = NewObject<UInterchangeSkeletalMeshFactoryNode>(DestinationContainer);
+		MeshNode->InitializeSkeletalMeshNode(NodeUniqueId, NodeDisplayLabel, USkeletalMesh::StaticClass()->GetName());
+		DestinationContainer->AddNode(MeshNode);
+
+		FillInterchangeGenericAssetsPipelineFromFbxMeshImportData(GenericAssetPipeline, SkeletalMeshImportData);
+
+		GenericAssetPipeline->CommonSkeletalMeshesAndAnimationsProperties->bImportMeshesInBoneHierarchy = SkeletalMeshImportData->bImportMeshesInBoneHierarchy;
+		GenericAssetPipeline->MeshPipeline->bImportMorphTargets = SkeletalMeshImportData->bImportMorphTargets;
+		GenericAssetPipeline->MeshPipeline->bImportVertexAttributes = SkeletalMeshImportData->bImportVertexAttributes;
+
+		//TODO integrate this new feature to interchange
+		//SkeletalMeshImportData->bKeepSectionsSeparate;
+
+		GenericAssetPipeline->MeshPipeline->bUpdateSkeletonReferencePose = SkeletalMeshImportData->bUpdateSkeletonReferencePose;
+		GenericAssetPipeline->CommonSkeletalMeshesAndAnimationsProperties->bUseT0AsRefPose = SkeletalMeshImportData->bUseT0AsRefPose;
+		if (SkeletalMeshImportData->ImportContentType == EFBXImportContentType::FBXICT_All)
+		{
+			GenericAssetPipeline->MeshPipeline->SkeletalMeshImportContentType = EInterchangeSkeletalMeshContentType::All;
+		}
+		else if (SkeletalMeshImportData->ImportContentType == EFBXImportContentType::FBXICT_Geometry)
+		{
+			GenericAssetPipeline->MeshPipeline->SkeletalMeshImportContentType = EInterchangeSkeletalMeshContentType::Geometry;
+		}
+		else if (SkeletalMeshImportData->ImportContentType == EFBXImportContentType::FBXICT_SkinningWeights)
+		{
+			GenericAssetPipeline->MeshPipeline->SkeletalMeshImportContentType = EInterchangeSkeletalMeshContentType::SkinningWeights;
+		}
+
+		if (SkeletalMeshImportData->LastImportContentType == EFBXImportContentType::FBXICT_All)
+		{
+			GenericAssetPipeline->MeshPipeline->LastSkeletalMeshImportContentType = EInterchangeSkeletalMeshContentType::All;
+		}
+		else if (SkeletalMeshImportData->LastImportContentType == EFBXImportContentType::FBXICT_Geometry)
+		{
+			GenericAssetPipeline->MeshPipeline->LastSkeletalMeshImportContentType = EInterchangeSkeletalMeshContentType::Geometry;
+		}
+		else if (SkeletalMeshImportData->LastImportContentType == EFBXImportContentType::FBXICT_SkinningWeights)
+		{
+			GenericAssetPipeline->MeshPipeline->LastSkeletalMeshImportContentType = EInterchangeSkeletalMeshContentType::SkinningWeights;
+		}
+
+		GenericAssetPipeline->MeshPipeline->MorphThresholdPosition = SkeletalMeshImportData->MorphThresholdPosition;
+		GenericAssetPipeline->MeshPipeline->ThresholdPosition = SkeletalMeshImportData->ThresholdPosition;
+		GenericAssetPipeline->MeshPipeline->ThresholdTangentNormal = SkeletalMeshImportData->ThresholdTangentNormal;
+		GenericAssetPipeline->MeshPipeline->ThresholdUV = SkeletalMeshImportData->ThresholdUV;
+
+		if (SkeletalMeshImportData->VertexColorImportOption == EVertexColorImportOption::Ignore)
+		{
+			GenericAssetPipeline->CommonMeshesProperties->VertexColorImportOption = EInterchangeVertexColorImportOption::IVCIO_Ignore;
+		}
+		else if (SkeletalMeshImportData->VertexColorImportOption == EVertexColorImportOption::Override)
+		{
+			GenericAssetPipeline->CommonMeshesProperties->VertexColorImportOption = EInterchangeVertexColorImportOption::IVCIO_Override;
+		}
+		else if (SkeletalMeshImportData->VertexColorImportOption == EVertexColorImportOption::Replace)
+		{
+			GenericAssetPipeline->CommonMeshesProperties->VertexColorImportOption = EInterchangeVertexColorImportOption::IVCIO_Replace;
+		}
+		GenericAssetPipeline->CommonMeshesProperties->VertexOverrideColor = SkeletalMeshImportData->VertexOverrideColor;
+	}
+
+	void FillInterchangeGenericAssetsPipelineFromFbxAnimSequenceImportData(UInterchangeGenericAssetsPipeline* GenericAssetPipeline
+		, const UFbxAnimSequenceImportData* AnimSequenceImportData
+		, UInterchangeBaseNodeContainer* DestinationContainer
+		, const FString& NodeUniqueId
+		, const FString& NodeDisplayLabel)
+	{
+		UInterchangeAnimSequenceFactoryNode* AnimationNode = NewObject<UInterchangeAnimSequenceFactoryNode>(DestinationContainer);
+		AnimationNode->InitializeAnimSequenceNode(NodeUniqueId, NodeDisplayLabel);
+		DestinationContainer->AddNode(AnimationNode);
+		switch (AnimSequenceImportData->AnimationLength)
+		{
+		case EFBXAnimationLengthImportType::FBXALIT_ExportedTime:
+		{
+			GenericAssetPipeline->AnimationPipeline->AnimationRange = EInterchangeAnimationRange::Timeline;
+			break;
+		}
+		case EFBXAnimationLengthImportType::FBXALIT_AnimatedKey:
+		{
+			GenericAssetPipeline->AnimationPipeline->AnimationRange = EInterchangeAnimationRange::Animated;
+			break;
+		}
+		case EFBXAnimationLengthImportType::FBXALIT_SetRange:
+		{
+			GenericAssetPipeline->AnimationPipeline->AnimationRange = EInterchangeAnimationRange::SetRange;
+			break;
+		}
+		}
+		GenericAssetPipeline->AnimationPipeline->bAddCurveMetadataToSkeleton = AnimSequenceImportData->bAddCurveMetadataToSkeleton;
+		GenericAssetPipeline->AnimationPipeline->bDeleteExistingCustomAttributeCurves = AnimSequenceImportData->bDeleteExistingCustomAttributeCurves;
+		GenericAssetPipeline->AnimationPipeline->bDeleteExistingMorphTargetCurves = AnimSequenceImportData->bDeleteExistingMorphTargetCurves;
+		GenericAssetPipeline->AnimationPipeline->bDeleteExistingNonCurveCustomAttributes = AnimSequenceImportData->bDeleteExistingNonCurveCustomAttributes;
+		GenericAssetPipeline->AnimationPipeline->bDoNotImportCurveWithZero = AnimSequenceImportData->bDoNotImportCurveWithZero;
+		GenericAssetPipeline->AnimationPipeline->bImportBoneTracks = AnimSequenceImportData->bImportBoneTracks;
+		GenericAssetPipeline->AnimationPipeline->bImportCustomAttribute = AnimSequenceImportData->bImportCustomAttribute;
+		GenericAssetPipeline->CommonSkeletalMeshesAndAnimationsProperties->bImportMeshesInBoneHierarchy = AnimSequenceImportData->bImportMeshesInBoneHierarchy;
+		GenericAssetPipeline->AnimationPipeline->bRemoveCurveRedundantKeys = AnimSequenceImportData->bRemoveRedundantKeys;
+		GenericAssetPipeline->AnimationPipeline->bSetMaterialDriveParameterOnCustomAttribute = AnimSequenceImportData->bSetMaterialDriveParameterOnCustomAttribute;
+		GenericAssetPipeline->AnimationPipeline->bSnapToClosestFrameBoundary = AnimSequenceImportData->bSnapToClosestFrameBoundary;
+		GenericAssetPipeline->AnimationPipeline->bUse30HzToBakeBoneAnimation = AnimSequenceImportData->bUseDefaultSampleRate;
+		GenericAssetPipeline->AnimationPipeline->CustomBoneAnimationSampleRate = AnimSequenceImportData->CustomSampleRate;
+		GenericAssetPipeline->AnimationPipeline->FrameImportRange = AnimSequenceImportData->FrameImportRange;
+		GenericAssetPipeline->AnimationPipeline->MaterialCurveSuffixes = AnimSequenceImportData->MaterialCurveSuffixes;
+		GenericAssetPipeline->AnimationPipeline->SourceAnimationName = AnimSequenceImportData->SourceAnimationName;
+	}
+
+	UAssetImportData* ConvertToLegacyFbx(UStaticMesh* StaticMesh, const UInterchangeAssetImportData* InterchangeSourceData)
+	{
+		if (!InterchangeSourceData || !StaticMesh)
+		{
+			return nullptr;
+		}
+
+		//Create a fbx asset import data and fill the options
+		UFbxStaticMeshImportData* DestinationStaticMeshImportData = NewObject<UFbxStaticMeshImportData>(StaticMesh);
+
+		if (!DestinationStaticMeshImportData)
+		{
+			return nullptr;
+		}
+
+		//Transfer the Source file information
+		TransferSourceFileInformation(InterchangeSourceData, DestinationStaticMeshImportData);
+
+		//Now find the generic asset pipeline
+		for (UObject* Pipeline : InterchangeSourceData->GetPipelines())
+		{
+			if (const UInterchangeGenericAssetsPipeline* GenericAssetPipeline = Cast<UInterchangeGenericAssetsPipeline>(Pipeline))
+			{
+
+				FillFbxAssetImportData(GenericAssetPipeline, DestinationStaticMeshImportData);
+				FillFbxMeshImportData(GenericAssetPipeline, DestinationStaticMeshImportData);
+
+				DestinationStaticMeshImportData->bAutoGenerateCollision = GenericAssetPipeline->MeshPipeline->bImportCollision;
+				DestinationStaticMeshImportData->bBuildNanite = GenericAssetPipeline->MeshPipeline->bBuildNanite;
+				DestinationStaticMeshImportData->bBuildReversedIndexBuffer = GenericAssetPipeline->MeshPipeline->bBuildReversedIndexBuffer;
+				DestinationStaticMeshImportData->bCombineMeshes = GenericAssetPipeline->MeshPipeline->bCombineStaticMeshes;
+				DestinationStaticMeshImportData->bGenerateLightmapUVs = GenericAssetPipeline->MeshPipeline->bGenerateLightmapUVs;
+				DestinationStaticMeshImportData->bOneConvexHullPerUCX = GenericAssetPipeline->MeshPipeline->bOneConvexHullPerUCX;
+				DestinationStaticMeshImportData->bRemoveDegenerates = GenericAssetPipeline->CommonMeshesProperties->bRemoveDegenerates;
+				DestinationStaticMeshImportData->DistanceFieldResolutionScale = GenericAssetPipeline->MeshPipeline->DistanceFieldResolutionScale;
+				DestinationStaticMeshImportData->StaticMeshLODGroup = GenericAssetPipeline->MeshPipeline->LodGroup;
+				if (GenericAssetPipeline->CommonMeshesProperties->VertexColorImportOption == EInterchangeVertexColorImportOption::IVCIO_Ignore)
+				{
+					DestinationStaticMeshImportData->VertexColorImportOption = EVertexColorImportOption::Ignore;
+				}
+				else if (GenericAssetPipeline->CommonMeshesProperties->VertexColorImportOption == EInterchangeVertexColorImportOption::IVCIO_Override)
+				{
+					DestinationStaticMeshImportData->VertexColorImportOption = EVertexColorImportOption::Override;
+				}
+				else if (GenericAssetPipeline->CommonMeshesProperties->VertexColorImportOption == EInterchangeVertexColorImportOption::IVCIO_Replace)
+				{
+					DestinationStaticMeshImportData->VertexColorImportOption = EVertexColorImportOption::Replace;
+				}
+				DestinationStaticMeshImportData->VertexOverrideColor = GenericAssetPipeline->CommonMeshesProperties->VertexOverrideColor;
+
+				//Fill the reimport material match data and section data
+				FImportMeshLodSectionsData SectionData;
+				for (const FStaticMaterial& Material : StaticMesh->GetStaticMaterials())
+				{
+					DestinationStaticMeshImportData->ImportMaterialOriginalNameData.Add(Material.ImportedMaterialSlotName);
+					SectionData.SectionOriginalMaterialName.Add(Material.ImportedMaterialSlotName);
+				}
+				DestinationStaticMeshImportData->ImportMeshLodData.Add(SectionData);
+			}
+		}
+		return DestinationStaticMeshImportData;
+	}
+
+	UAssetImportData* ConvertToLegacyFbx(USkeletalMesh* SkeletalMesh, const UInterchangeAssetImportData* InterchangeSourceData)
+	{
+		if (!InterchangeSourceData || !SkeletalMesh)
+		{
+			return nullptr;
+		}
+
+		//Create a fbx asset import data and fill the options
+		UFbxSkeletalMeshImportData* DestinationSkeletalMeshImportData = NewObject<UFbxSkeletalMeshImportData>(SkeletalMesh);
+
+		if (!DestinationSkeletalMeshImportData)
+		{
+			return nullptr;
+		}
+
+		//Transfer the Source file information
+		TransferSourceFileInformation(InterchangeSourceData, DestinationSkeletalMeshImportData);
+
+		//Now find the generic asset pipeline
+		for (UObject* Pipeline : InterchangeSourceData->GetPipelines())
+		{
+			if (const UInterchangeGenericAssetsPipeline* GenericAssetPipeline = Cast<UInterchangeGenericAssetsPipeline>(Pipeline))
+			{
+				FillFbxAssetImportData(GenericAssetPipeline, DestinationSkeletalMeshImportData);
+				FillFbxMeshImportData(GenericAssetPipeline, DestinationSkeletalMeshImportData);
+				DestinationSkeletalMeshImportData->bImportMeshesInBoneHierarchy = GenericAssetPipeline->CommonSkeletalMeshesAndAnimationsProperties->bImportMeshesInBoneHierarchy;
+				DestinationSkeletalMeshImportData->bImportMorphTargets = GenericAssetPipeline->MeshPipeline->bImportMorphTargets;
+				DestinationSkeletalMeshImportData->bImportVertexAttributes = GenericAssetPipeline->MeshPipeline->bImportVertexAttributes;
+
+				//TODO integrate this new feature to interchange
+				//DestinationSkeletalMeshImportData->bKeepSectionsSeparate = ;
+
+				DestinationSkeletalMeshImportData->bPreserveSmoothingGroups = true;
+				DestinationSkeletalMeshImportData->bUpdateSkeletonReferencePose = GenericAssetPipeline->MeshPipeline->bUpdateSkeletonReferencePose;
+				DestinationSkeletalMeshImportData->bUseT0AsRefPose = GenericAssetPipeline->CommonSkeletalMeshesAndAnimationsProperties->bUseT0AsRefPose;
+				if (GenericAssetPipeline->MeshPipeline->SkeletalMeshImportContentType == EInterchangeSkeletalMeshContentType::All)
+				{
+					DestinationSkeletalMeshImportData->ImportContentType = EFBXImportContentType::FBXICT_All;
+				}
+				else if (GenericAssetPipeline->MeshPipeline->SkeletalMeshImportContentType == EInterchangeSkeletalMeshContentType::Geometry)
+				{
+					DestinationSkeletalMeshImportData->ImportContentType = EFBXImportContentType::FBXICT_Geometry;
+				}
+				else if (GenericAssetPipeline->MeshPipeline->SkeletalMeshImportContentType == EInterchangeSkeletalMeshContentType::SkinningWeights)
+				{
+					DestinationSkeletalMeshImportData->ImportContentType = EFBXImportContentType::FBXICT_SkinningWeights;
+				}
+
+				if (GenericAssetPipeline->MeshPipeline->LastSkeletalMeshImportContentType == EInterchangeSkeletalMeshContentType::All)
+				{
+					DestinationSkeletalMeshImportData->LastImportContentType = EFBXImportContentType::FBXICT_All;
+				}
+				else if (GenericAssetPipeline->MeshPipeline->LastSkeletalMeshImportContentType == EInterchangeSkeletalMeshContentType::Geometry)
+				{
+					DestinationSkeletalMeshImportData->LastImportContentType = EFBXImportContentType::FBXICT_Geometry;
+				}
+				else if (GenericAssetPipeline->MeshPipeline->LastSkeletalMeshImportContentType == EInterchangeSkeletalMeshContentType::SkinningWeights)
+				{
+					DestinationSkeletalMeshImportData->LastImportContentType = EFBXImportContentType::FBXICT_SkinningWeights;
+				}
+
+				DestinationSkeletalMeshImportData->MorphThresholdPosition = GenericAssetPipeline->MeshPipeline->MorphThresholdPosition;
+				DestinationSkeletalMeshImportData->ThresholdPosition = GenericAssetPipeline->MeshPipeline->ThresholdPosition;
+				DestinationSkeletalMeshImportData->ThresholdTangentNormal = GenericAssetPipeline->MeshPipeline->ThresholdTangentNormal;
+				DestinationSkeletalMeshImportData->ThresholdUV = GenericAssetPipeline->MeshPipeline->ThresholdUV;
+
+				if (GenericAssetPipeline->CommonMeshesProperties->VertexColorImportOption == EInterchangeVertexColorImportOption::IVCIO_Ignore)
+				{
+					DestinationSkeletalMeshImportData->VertexColorImportOption = EVertexColorImportOption::Ignore;
+				}
+				else if (GenericAssetPipeline->CommonMeshesProperties->VertexColorImportOption == EInterchangeVertexColorImportOption::IVCIO_Override)
+				{
+					DestinationSkeletalMeshImportData->VertexColorImportOption = EVertexColorImportOption::Override;
+				}
+				else if (GenericAssetPipeline->CommonMeshesProperties->VertexColorImportOption == EInterchangeVertexColorImportOption::IVCIO_Replace)
+				{
+					DestinationSkeletalMeshImportData->VertexColorImportOption = EVertexColorImportOption::Replace;
+				}
+				DestinationSkeletalMeshImportData->VertexOverrideColor = GenericAssetPipeline->CommonMeshesProperties->VertexOverrideColor;
+
+				//Fill the reimport material match data and section data
+				FImportMeshLodSectionsData SectionData;
+				for (const FSkeletalMaterial& Material : SkeletalMesh->GetMaterials())
+				{
+					DestinationSkeletalMeshImportData->ImportMaterialOriginalNameData.Add(Material.ImportedMaterialSlotName);
+					SectionData.SectionOriginalMaterialName.Add(Material.ImportedMaterialSlotName);
+				}
+				DestinationSkeletalMeshImportData->ImportMeshLodData.Add(SectionData);
+			}
+		}
+		return DestinationSkeletalMeshImportData;
+	}
+
+	UAssetImportData* ConvertToLegacyFbx(UAnimSequence* AnimSequence, const UInterchangeAssetImportData* InterchangeSourceData)
+	{
+		if (!InterchangeSourceData || !AnimSequence)
+		{
+			return nullptr;
+		}
+
+		//Create a fbx asset import data and fill the options
+		UFbxAnimSequenceImportData* DestinationAnimSequenceImportData = NewObject<UFbxAnimSequenceImportData>(AnimSequence);
+
+		if (!DestinationAnimSequenceImportData)
+		{
+			return nullptr;
+		}
+
+		//Transfer the Source file information
+		TransferSourceFileInformation(InterchangeSourceData, DestinationAnimSequenceImportData);
+
+		//Now find the generic asset pipeline
+		for (UObject* Pipeline : InterchangeSourceData->GetPipelines())
+		{
+			if (const UInterchangeGenericAssetsPipeline* GenericAssetPipeline = Cast<UInterchangeGenericAssetsPipeline>(Pipeline))
+			{
+				FillFbxAssetImportData(GenericAssetPipeline, DestinationAnimSequenceImportData);
+
+				switch (GenericAssetPipeline->AnimationPipeline->AnimationRange)
+				{
+				case EInterchangeAnimationRange::Timeline:
+				{
+					DestinationAnimSequenceImportData->AnimationLength = EFBXAnimationLengthImportType::FBXALIT_ExportedTime;
+					break;
+				}
+				case EInterchangeAnimationRange::Animated:
+				{
+					DestinationAnimSequenceImportData->AnimationLength = EFBXAnimationLengthImportType::FBXALIT_AnimatedKey;
+					break;
+				}
+				case EInterchangeAnimationRange::SetRange:
+				{
+					DestinationAnimSequenceImportData->AnimationLength = EFBXAnimationLengthImportType::FBXALIT_SetRange;
+					break;
+				}
+				}
+				DestinationAnimSequenceImportData->bAddCurveMetadataToSkeleton = GenericAssetPipeline->AnimationPipeline->bAddCurveMetadataToSkeleton;
+				DestinationAnimSequenceImportData->bDeleteExistingCustomAttributeCurves = GenericAssetPipeline->AnimationPipeline->bDeleteExistingCustomAttributeCurves;
+				DestinationAnimSequenceImportData->bDeleteExistingMorphTargetCurves = GenericAssetPipeline->AnimationPipeline->bDeleteExistingMorphTargetCurves;
+				DestinationAnimSequenceImportData->bDeleteExistingNonCurveCustomAttributes = GenericAssetPipeline->AnimationPipeline->bDeleteExistingNonCurveCustomAttributes;
+				DestinationAnimSequenceImportData->bDoNotImportCurveWithZero = GenericAssetPipeline->AnimationPipeline->bDoNotImportCurveWithZero;
+				DestinationAnimSequenceImportData->bImportBoneTracks = GenericAssetPipeline->AnimationPipeline->bImportBoneTracks;
+				DestinationAnimSequenceImportData->bImportCustomAttribute = GenericAssetPipeline->AnimationPipeline->bImportCustomAttribute;
+				DestinationAnimSequenceImportData->bImportMeshesInBoneHierarchy = GenericAssetPipeline->CommonSkeletalMeshesAndAnimationsProperties->bImportMeshesInBoneHierarchy;
+				DestinationAnimSequenceImportData->bPreserveLocalTransform = false;
+				DestinationAnimSequenceImportData->bRemoveRedundantKeys = GenericAssetPipeline->AnimationPipeline->bRemoveCurveRedundantKeys;
+				DestinationAnimSequenceImportData->bSetMaterialDriveParameterOnCustomAttribute = GenericAssetPipeline->AnimationPipeline->bSetMaterialDriveParameterOnCustomAttribute;
+				DestinationAnimSequenceImportData->bSnapToClosestFrameBoundary = GenericAssetPipeline->AnimationPipeline->bSnapToClosestFrameBoundary;
+				DestinationAnimSequenceImportData->bUseDefaultSampleRate = GenericAssetPipeline->AnimationPipeline->bUse30HzToBakeBoneAnimation;
+				DestinationAnimSequenceImportData->CustomSampleRate = GenericAssetPipeline->AnimationPipeline->CustomBoneAnimationSampleRate;
+				DestinationAnimSequenceImportData->FrameImportRange = GenericAssetPipeline->AnimationPipeline->FrameImportRange;
+				DestinationAnimSequenceImportData->MaterialCurveSuffixes = GenericAssetPipeline->AnimationPipeline->MaterialCurveSuffixes;
+				DestinationAnimSequenceImportData->SourceAnimationName = GenericAssetPipeline->AnimationPipeline->SourceAnimationName;
+			}
+		}
+		return DestinationAnimSequenceImportData;
+	}
+
+	UAssetImportData* ConvertToInterchange(UObject* Obj, const UFbxAssetImportData* FbxAssetImportData)
+	{
+		if (!FbxAssetImportData || !Obj)
+		{
+			return nullptr;
+		}
+		//Create a fbx asset import data and fill the options
+		UInterchangeAssetImportData* DestinationData = NewObject<UInterchangeAssetImportData>(Obj);
+		//Transfer the Source file information
+		TransferSourceFileInformation(FbxAssetImportData, DestinationData);
+
+		//Create a container
+		UInterchangeBaseNodeContainer* DestinationContainer = NewObject<UInterchangeBaseNodeContainer>(DestinationData);
+		const FString BasePathToRemove = FPaths::GetBaseFilename(FbxAssetImportData->GetFirstFilename()) + TEXT("_");
+		FString NodeDisplayLabel = Obj->GetName();
+		if (NodeDisplayLabel.StartsWith(BasePathToRemove))
+		{
+			NodeDisplayLabel.RightInline(NodeDisplayLabel.Len() - BasePathToRemove.Len());
+		}
+		const FString NodeUniqueId = FGuid::NewGuid().ToString(EGuidFormats::Base36Encoded) + TEXT("_") + NodeDisplayLabel;
+
+		TArray<UObject*> Pipelines;
+		UInterchangeGenericAssetsPipeline* GenericAssetPipeline = NewObject<UInterchangeGenericAssetsPipeline>(DestinationData);
+		Pipelines.Add(GenericAssetPipeline);
+		DestinationData->SetPipelines(Pipelines);
+
+		GenericAssetPipeline->ReimportStrategy = EReimportStrategyFlags::ApplyNoProperties;
+		GenericAssetPipeline->ImportOffsetRotation = FbxAssetImportData->ImportRotation;
+		GenericAssetPipeline->ImportOffsetTranslation = FbxAssetImportData->ImportTranslation;
+		GenericAssetPipeline->ImportOffsetUniformScale = FbxAssetImportData->ImportUniformScale;
+
+		if (const UFbxStaticMeshImportData* LegacyStaticMeshImportData = Cast<UFbxStaticMeshImportData>(FbxAssetImportData))
+		{
+			check(Obj->IsA<UStaticMesh>());
+			FillInterchangeGenericAssetsPipelineFromFbxStaticMeshImportData(GenericAssetPipeline
+				, LegacyStaticMeshImportData
+				, DestinationContainer
+				, NodeUniqueId
+				, NodeDisplayLabel);
+		}
+		else if (const UFbxSkeletalMeshImportData* LegacySkeletalMeshImportData = Cast<UFbxSkeletalMeshImportData>(FbxAssetImportData))
+		{
+			check(Obj->IsA<USkeletalMesh>());
+			FillInterchangeGenericAssetsPipelineFromFbxSkeletalMeshImportData(GenericAssetPipeline
+				, LegacySkeletalMeshImportData
+				, DestinationContainer
+				, NodeUniqueId
+				, NodeDisplayLabel);
+		}
+		else if (const UFbxAnimSequenceImportData* LegacyAnimSequenceImportData = Cast<UFbxAnimSequenceImportData>(FbxAssetImportData))
+		{
+			check(Obj->IsA<UAnimSequence>());
+			FillInterchangeGenericAssetsPipelineFromFbxAnimSequenceImportData(GenericAssetPipeline
+				, LegacyAnimSequenceImportData
+				, DestinationContainer
+				, NodeUniqueId
+				, NodeDisplayLabel);
+		}
+
+		if (UInterchangeFactoryBaseNode* DestinationFactoryNode = DestinationContainer->GetFactoryNode(NodeUniqueId))
+		{
+			DestinationFactoryNode->SetReimportStrategyFlags(EReimportStrategyFlags::ApplyNoProperties);
+			DestinationFactoryNode->SetCustomReferenceObject(Obj);
+			DestinationData->SetNodeContainer(DestinationContainer);
+			DestinationData->NodeUniqueID = NodeUniqueId;
+		}
+		return DestinationData;
+	}
+
+	UAssetImportData* ConvertData(UObject* Obj, UAssetImportData* SourceData, const bool bInterchangeSupportTargetExtension)
+	{
+		if (const UInterchangeAssetImportData* InterchangeSourceData = Cast<UInterchangeAssetImportData>(SourceData))
+		{
+			if (bInterchangeSupportTargetExtension)
+			{
+				//This converter do not convert Interchange to Interchange
+				return nullptr;
+			}
+
+			//Do not convert scene data
+			if (InterchangeSourceData->SceneImportAsset.IsValid())
+			{
+				return nullptr;
+			}
+
+			//Convert Interchange import data to Legacy Fbx Import data
+			if (UStaticMesh* StaticMesh = Cast<UStaticMesh>(Obj))
+			{
+				return ConvertToLegacyFbx(StaticMesh, InterchangeSourceData);
+			}
+			else if (USkeletalMesh* SkeletalMesh = Cast<USkeletalMesh>(Obj))
+			{
+				return ConvertToLegacyFbx(SkeletalMesh, InterchangeSourceData);
+			}
+			else if (UAnimSequence* AnimSequence = Cast<UAnimSequence>(Obj))
+			{
+				return ConvertToLegacyFbx(AnimSequence, InterchangeSourceData);
+			}
+		}
+		else if (const UFbxAssetImportData* LegacyFbxSourceData = Cast<UFbxAssetImportData>(SourceData))
+		{
+			if (!bInterchangeSupportTargetExtension)
+			{
+				//This converter do not convert Legacy Fbx to other format then Interchange.
+				//This is probably a conversion from Legacy Fbx to Legacy Fbx which we do not need to do
+				return nullptr;
+			}
+
+			//Do not convert scene data
+			if (LegacyFbxSourceData->bImportAsScene)
+			{
+				return nullptr;
+			}
+
+			//Convert Legacy Fbx import data to Interchange Import data
+			return ConvertToInterchange(Obj, LegacyFbxSourceData);
+		}
+		return nullptr;
+	}
+} //ns: UE::Interchange::Private
+
+bool UInterchangeFbxAssetImportDataConverter::ConvertImportData(UObject* Obj, const FString& TargetExtension) const
+{
+	bool bResult = false;
+	const FString TargetExtensionLower = TargetExtension.ToLower();
+	bool bUseInterchangeFramework = UInterchangeManager::IsInterchangeImportEnabled();;
+	UInterchangeManager& InterchangeManager = UInterchangeManager::GetInterchangeManager();
+	
+	UAssetImportData* OldAssetData = nullptr;
+	TArray<FString> InterchangeSupportedExtensions;
+	if (Obj->IsA(UStaticMesh::StaticClass()) || Obj->IsA(USkeletalMesh::StaticClass()))
+	{
+		InterchangeSupportedExtensions = InterchangeManager.GetSupportedAssetTypeFormats(EInterchangeTranslatorAssetType::Meshes);
+	}
+	else if (Obj->IsA(UAnimSequence::StaticClass()))
+	{
+		InterchangeSupportedExtensions = InterchangeManager.GetSupportedAssetTypeFormats(EInterchangeTranslatorAssetType::Animations);
+	}
+	//Remove the detail of the extensions
+	for (FString& Extension : InterchangeSupportedExtensions)
+	{
+		int32 FindIndex = INDEX_NONE;
+		Extension.FindChar(';', FindIndex);
+		if (FindIndex != INDEX_NONE && FindIndex < Extension.Len() && FindIndex > 0)
+		{
+			Extension.LeftInline(FindIndex);
+		}
+	}
+	const bool bInterchangeSupportTargetExtension = bUseInterchangeFramework && InterchangeSupportedExtensions.Contains(TargetExtensionLower);
+	
+	if (TargetExtensionLower.Equals(TEXT("fbx")) || bInterchangeSupportTargetExtension)
+	{
+		if (UStaticMesh* StaticMesh = Cast<UStaticMesh>(Obj))
+		{
+			if (UAssetImportData* ConvertedAssetData = UE::Interchange::Private::ConvertData(StaticMesh, StaticMesh->GetAssetImportData(), bInterchangeSupportTargetExtension))
+			{
+				OldAssetData = StaticMesh->GetAssetImportData();
+				StaticMesh->SetAssetImportData(ConvertedAssetData);
+				bResult = true;
+			}
+		}
+		else if (USkeletalMesh* SkeletalMesh = Cast<USkeletalMesh>(Obj))
+		{
+			if (UAssetImportData* ConvertedAssetData = UE::Interchange::Private::ConvertData(SkeletalMesh, SkeletalMesh->GetAssetImportData(), bInterchangeSupportTargetExtension))
+			{
+				OldAssetData = SkeletalMesh->GetAssetImportData();
+				SkeletalMesh->SetAssetImportData(ConvertedAssetData);
+				bResult = true;
+			}
+		}
+		else if (UAnimSequence* AnimSequence = Cast<UAnimSequence>(Obj))
+		{
+			if (UAssetImportData* ConvertedAssetData = UE::Interchange::Private::ConvertData(AnimSequence, AnimSequence->AssetImportData, bInterchangeSupportTargetExtension))
+			{
+				OldAssetData = AnimSequence->AssetImportData;
+				AnimSequence->AssetImportData = ConvertedAssetData;
+				bResult = true;
+			}
+		}
+	}
+	
+	//Make sure old import asset data will be deleted by the next garbage collect
+	if (bResult && OldAssetData)
+	{
+		OldAssetData->Rename(nullptr, GetTransientPackage(), REN_DoNotDirty | REN_DontCreateRedirectors | REN_NonTransactional | REN_ForceNoResetLoaders);
+		OldAssetData->ClearFlags(RF_Public | RF_Standalone);
+	}
+
+	return bResult;
+}
+
+
diff --git a/Engine/Plugins/Interchange/Editor/Source/InterchangeEditor/Public/InterchangeFbxAssetImportDataConverter.h b/Engine/Plugins/Interchange/Editor/Source/InterchangeEditor/Public/InterchangeFbxAssetImportDataConverter.h
new file mode 100644
--- /dev/null
+++ b/Engine/Plugins/Interchange/Editor/Source/InterchangeEditor/Public/InterchangeFbxAssetImportDataConverter.h
@@ -0,0 +1,19 @@
+// Copyright Epic Games, Inc. All Rights Reserved.
+
+#pragma once
+
+#include "InterchangeAssetImportData.h"
+#include "UObject/Class.h"
+#include "UObject/Object.h"
+#include "UObject/ObjectMacros.h"
+
+#include "InterchangeFbxAssetImportDataConverter.generated.h"
+
+UCLASS()
+class INTERCHANGEEDITOR_API UInterchangeFbxAssetImportDataConverter : public UInterchangeAssetImportDataConverterBase
+{
+	GENERATED_BODY()
+
+public:
+	virtual bool ConvertImportData(UObject* Object, const FString& ToExtension) const override;
+};
diff --git a/Engine/Source/Editor/UnrealEd/Private/Editor.cpp b/Engine/Source/Editor/UnrealEd/Private/Editor.cpp
--- a/Engine/Source/Editor/UnrealEd/Private/Editor.cpp
+++ b/Engine/Source/Editor/UnrealEd/Private/Editor.cpp
@@ -71,6 +71,7 @@
 #include "AutoReimport/AutoReimportUtilities.h"
 #include "AssetToolsModule.h"
 
+#include "InterchangeAssetImportData.h"
 #include "InterchangeManager.h"
 #include "InterchangeResultsContainer.h"
 
@@ -294,7 +295,9 @@ UE::Interchange::FAssetImportResultRef FReimportManager::ReimportAsync(UObject*
 
 	bool bUseInterchangeFramework = UInterchangeManager::IsInterchangeImportEnabled();;
 	UInterchangeManager& InterchangeManager = UInterchangeManager::GetInterchangeManager();
-
+	
+	const int32 RealSourceFileIndex = SourceFileIndex == INDEX_NONE ? 0 : SourceFileIndex;
+	
 	bool bSuccess = false;
 	if ( Obj )
 	{
@@ -344,7 +347,6 @@ UE::Interchange::FAssetImportResultRef FReimportManager::ReimportAsync(UObject*
 			}
 			else
 			{
-				int32 RealSourceFileIndex = SourceFileIndex == INDEX_NONE ? 0 : SourceFileIndex;
 				if (bForceNewFile)
 				{
 					if (SourceFilenames.IsValidIndex(RealSourceFileIndex))
@@ -409,10 +411,41 @@ UE::Interchange::FAssetImportResultRef FReimportManager::ReimportAsync(UObject*
 			{
 				// Reimporting the asset from a new file
 				CanReimportHandler->SetReimportPaths(Obj, PreferredReimportFile, SourceFileIndex);
+				//Update the local source file
+				if (SourceFilenames.IsValidIndex(RealSourceFileIndex))
+				{
+					SourceFilenames[RealSourceFileIndex] = PreferredReimportFile;
+				}
 			}
 
 			if ( bValidSourceFilename )
 			{
+				//Convert the import data if it's needed and choose a new valid reimport handler after the conversion is done.
+				//This allow us to re-import:
+				// Interchange -> Legacy Fbx    ---> Asset was imported with Interchange (gltf, fbx, obj, ...), Interchange is turn off for fbx and the provided source file is fbx
+				// Legacy Fbx -> Interchange    ---> Asset was imported with Legacy Fbx, the file use for re-import is supported by Interchange (fbx, gltf, obj, ...)
+				{
+					const FString ReimportFilename = SourceFilenames.IsValidIndex(RealSourceFileIndex) ? SourceFilenames[RealSourceFileIndex] : FString();
+					const FString ReimportFilenameExtension = FPaths::GetExtension(ReimportFilename).ToLower();
+					//Convertion will return false if there is no conversion to do.
+					if (InterchangeManager.ConvertImportData(Obj, ReimportFilenameExtension))
+					{
+						for (int32 NewFileHandlerIndex = 0; NewFileHandlerIndex < Handlers.Num(); ++NewFileHandlerIndex)
+						{
+							SourceFilenames.Empty();
+							if (!PreferredReimportFile.IsEmpty())
+							{
+								Handlers[NewFileHandlerIndex]->SetPreferredReimportPath(PreferredReimportFile);
+							}
+							if (Handlers[NewFileHandlerIndex]->CanReimport(Obj, SourceFilenames))
+							{
+								CanReimportHandler = Handlers[NewFileHandlerIndex];
+								break;
+							}
+						}
+					}
+				}
+
 				if (bUseInterchangeFramework && CanReimportHandler->IsInterchangeFactory())
 				{
 					// Make sure SourceFilenames reflects the source filenames in Obj
@@ -421,7 +454,6 @@ UE::Interchange::FAssetImportResultRef FReimportManager::ReimportAsync(UObject*
 					{
 						check( SourceFilenames.Num() > 0 );
 
-						int32 RealSourceFileIndex = SourceFileIndex == INDEX_NONE ? 0 : SourceFileIndex;
 						int32 RealValidSourceFileIndex = SourceFilenames.IsValidIndex(RealSourceFileIndex) ? RealSourceFileIndex : 0;
 						UE::Interchange::FScopedSourceData ScopedSourceData(SourceFilenames[RealValidSourceFileIndex]);
 						CanReimportHandler->SetReimportSourceIndex(Obj, SourceFileIndex);
diff --git a/Engine/Source/Runtime/Interchange/Engine/Private/InterchangeImportCommon.cpp b/Engine/Source/Runtime/Interchange/Engine/Private/InterchangeImportCommon.cpp
--- a/Engine/Source/Runtime/Interchange/Engine/Private/InterchangeImportCommon.cpp
+++ b/Engine/Source/Runtime/Interchange/Engine/Private/InterchangeImportCommon.cpp
@@ -303,7 +303,7 @@ namespace UE::Interchange
 										, UInterchangeFactoryBaseNode* PipelineAssetNode)
 	{
 		TRACE_CPUPROFILER_EVENT_SCOPE(UE::Interchange::FFactoryCommon::ApplyReimportStrategyToAsset)
-		if (!ensure(PreviousAssetNode) || !ensure(PipelineAssetNode) || !ensure(CurrentAssetNode))
+		if (!ensure(PipelineAssetNode) || !ensure(CurrentAssetNode))
 		{
 			return;
 		}
@@ -325,6 +325,11 @@ namespace UE::Interchange
 				
 			case EReimportStrategyFlags::ApplyEditorChangedProperties:
 			{
+				if (!PreviousAssetNode)
+				{
+					UE_LOG(LogInterchangeEngine, Error, TEXT("Cannot apply the re-import strategy for asset [%s], because there is no previous asset node in the import data."), *Asset->GetName());
+					return;
+				}
 				TArray<FAttributeKey> RemovedAttributes;
 				TArray<FAttributeKey> AddedAttributes;
 				TArray<FAttributeKey> ModifiedAttributes;
diff --git a/Engine/Source/Runtime/Interchange/Engine/Private/InterchangeManager.cpp b/Engine/Source/Runtime/Interchange/Engine/Private/InterchangeManager.cpp
--- a/Engine/Source/Runtime/Interchange/Engine/Private/InterchangeManager.cpp
+++ b/Engine/Source/Runtime/Interchange/Engine/Private/InterchangeManager.cpp
@@ -1000,6 +1000,40 @@ bool UInterchangeManager::RegisterWriter(const UClass* WriterClass)
 	return true;
 }
 
+bool UInterchangeManager::RegisterImportDataConverter(const UClass* Converter)
+{
+#if WITH_EDITOR
+	if (!Converter)
+	{
+		return false;
+	}
+
+	if (RegisteredConverters.Contains(Converter))
+	{
+		return true;
+	}
+	UInterchangeAssetImportDataConverterBase* ConverterToRegister = NewObject<UInterchangeAssetImportDataConverterBase>(GetTransientPackage(), Converter, NAME_None);
+	if (!ConverterToRegister)
+	{
+		return false;
+	}
+	RegisteredConverters.Add(Converter, ConverterToRegister);
+#endif
+	return true;
+}
+
+bool UInterchangeManager::ConvertImportData(UObject* Object, const FString& Extension)
+{
+	for (TPair<TObjectPtr<const UClass>, TObjectPtr<UInterchangeAssetImportDataConverterBase>> RegisteredConverter : RegisteredConverters)
+	{
+		if (RegisteredConverter.Value->ConvertImportData(Object, Extension))
+		{
+			return true;
+		}
+	}
+	return false;
+}
+
 TArray<FString> UInterchangeManager::GetSupportedFormats(const EInterchangeTranslatorType ForTranslatorType) const
 {
 	TArray<FString> FileExtensions;
@@ -1085,6 +1119,12 @@ TArray<FString> UInterchangeManager::GetSupportedFormatsForObject(const UObject*
 		break;
 	}
 
+	//Make sure we return lower case extensions
+	for (FString& Extension : FileExtensions)
+	{
+		Extension.ToLowerInline();
+	}
+
 	return FileExtensions;
 }
 
diff --git a/Engine/Source/Runtime/Interchange/Engine/Public/InterchangeAssetImportData.h b/Engine/Source/Runtime/Interchange/Engine/Public/InterchangeAssetImportData.h
--- a/Engine/Source/Runtime/Interchange/Engine/Public/InterchangeAssetImportData.h
+++ b/Engine/Source/Runtime/Interchange/Engine/Public/InterchangeAssetImportData.h
@@ -160,3 +160,24 @@ private:
 	mutable TArray64<uint8> CachedNodeContainer;
 	mutable TArray<TPair<FString, FString>> CachedPipelines; //Class, Data(serialized JSON) pair
 };
+
+/**
+ * Base class to create an asset import data converter
+ */
+UCLASS(Abstract, MinimalAPI)
+class UInterchangeAssetImportDataConverterBase : public UObject
+{
+	GENERATED_BODY()
+public:
+	/**
+	 * Convert the asset import data from the one that is in the Object to
+	 * one supporting the target extension (i.e. legacy Fbx to Interchange or vice versa)
+	 * The function should return true only if it has convert the asset import data, false otherwise.
+	 * 
+	 * The system will call all object deriving from this class until one convert the data.
+	 */
+	virtual bool ConvertImportData(UObject* Object, const FString& TargetExtension) const
+	{
+		return false;
+	}
+};
diff --git a/Engine/Source/Runtime/Interchange/Engine/Public/InterchangeManager.h b/Engine/Source/Runtime/Interchange/Engine/Public/InterchangeManager.h
--- a/Engine/Source/Runtime/Interchange/Engine/Public/InterchangeManager.h
+++ b/Engine/Source/Runtime/Interchange/Engine/Public/InterchangeManager.h
@@ -12,6 +12,7 @@
 #include "HAL/CriticalSection.h"
 #include "HAL/Thread.h"
 #include "HAL/ThreadSafeBool.h"
+#include "InterchangeAssetImportData.h"
 #include "InterchangePipelineConfigurationBase.h"
 #include "InterchangeResultsContainer.h"
 #include "InterchangeSourceData.h"
@@ -467,6 +468,23 @@ public:
 	 */
 	INTERCHANGEENGINE_API bool RegisterWriter(const UClass* Writer);
 
+	/**
+	 * Any converter must register to the manager
+	 * @Param Converter - The UClass of the converter you want to register
+	 * @return true if the converter class can be register false otherwise.
+	 *
+	 * @Note if you register multiple time the same class it will return true for every call
+	 */
+	INTERCHANGEENGINE_API bool RegisterImportDataConverter(const UClass* Converter);
+
+	/**
+	 * Call all the registered converter, if one converter want 
+	 * @Param Object - The Object to convert the import data
+	 * @Param Extension - The file extension we want to import
+	 * @return true if one of the converter has convert the data. False otherwise.
+	 */
+	INTERCHANGEENGINE_API bool ConvertImportData(UObject* Object, const FString& Extension);
+
 	/**
 	 * Returns the list of supported formats for a given translator type.
 	 */
@@ -691,6 +709,10 @@ private:
 	UPROPERTY()
 	TMap<TObjectPtr<const UClass>, TObjectPtr<UInterchangeWriterBase> > RegisteredWriters;
 
+	//The manager will create only one converter per type
+	UPROPERTY()
+	TMap<TObjectPtr<const UClass>, TObjectPtr<UInterchangeAssetImportDataConverterBase> > RegisteredConverters;
+
 	//If interchange is currently importing we have a timer to watch the cancel and we block GC 
 	FThreadSafeBool bIsActive = false;
 
