Fix documentation link for importers

#rnx

diff --git a/Engine/Plugins/Animation/RigLogic/Source/RigLogicEditor/Private/DNAAssetImportWindow.cpp b/Engine/Plugins/Animation/RigLogic/Source/RigLogicEditor/Private/DNAAssetImportWindow.cpp
--- a/Engine/Plugins/Animation/RigLogic/Source/RigLogicEditor/Private/DNAAssetImportWindow.cpp
+++ b/Engine/Plugins/Animation/RigLogic/Source/RigLogicEditor/Private/DNAAssetImportWindow.cpp
@@ -82,7 +82,7 @@ void SDNAAssetImportWindow::Construct(const FArguments& InArgs)
 			.SlotPadding(2)
 		+ SUniformGridPanel::Slot(0, 0)
 		[
-			IDocumentation::Get()->CreateAnchor(FString("Engine/Content/FBX/ImportOptions"))
+			IDocumentation::Get()->CreateAnchor(FString("fbx-import-options-reference-in-unreal-engine"))
 		]
 	+ SUniformGridPanel::Slot(1, 0)
 		[
diff --git a/Engine/Plugins/Interchange/Editor/Source/Pipelines/Private/SInterchangePipelineConfigurationDialog.cpp b/Engine/Plugins/Interchange/Editor/Source/Pipelines/Private/SInterchangePipelineConfigurationDialog.cpp
--- a/Engine/Plugins/Interchange/Editor/Source/Pipelines/Private/SInterchangePipelineConfigurationDialog.cpp
+++ b/Engine/Plugins/Interchange/Editor/Source/Pipelines/Private/SInterchangePipelineConfigurationDialog.cpp
@@ -419,7 +419,7 @@ void SInterchangePipelineConfigurationDialog::Construct(const FArguments& InArgs
 				.Padding(4.f, 0.f)
 				.AutoWidth()
 				[
-					IDocumentation::Get()->CreateAnchor(FString("Engine/Content/Interchange/PipelineConfiguration"))
+					IDocumentation::Get()->CreateAnchor(FString("interchange-framework-in-unreal-engine"))
 				]
 				+ SHorizontalBox::Slot()
 				.Padding(4.f, 0.f)
diff --git a/Engine/Source/Editor/UnrealEd/Private/Fbx/FbxOptionWindow.cpp b/Engine/Source/Editor/UnrealEd/Private/Fbx/FbxOptionWindow.cpp
--- a/Engine/Source/Editor/UnrealEd/Private/Fbx/FbxOptionWindow.cpp
+++ b/Engine/Source/Editor/UnrealEd/Private/Fbx/FbxOptionWindow.cpp
@@ -149,7 +149,7 @@ void SFbxOptionWindow::Construct(const FArguments& InArgs)
 			.VAlign(VAlign_Center)
 			.AutoWidth()
 			[
-				IDocumentation::Get()->CreateAnchor(FString("Engine/Content/FBX/ImportOptions"))
+				IDocumentation::Get()->CreateAnchor(FString("fbx-import-options-reference-in-unreal-engine"))
 			]
 			+ SHorizontalBox::Slot()
 			.HAlign(HAlign_Right)
diff --git a/Engine/Source/Editor/UnrealEd/Private/Fbx/SFbxSceneOptionWindow.cpp b/Engine/Source/Editor/UnrealEd/Private/Fbx/SFbxSceneOptionWindow.cpp
--- a/Engine/Source/Editor/UnrealEd/Private/Fbx/SFbxSceneOptionWindow.cpp
+++ b/Engine/Source/Editor/UnrealEd/Private/Fbx/SFbxSceneOptionWindow.cpp
@@ -1294,7 +1294,7 @@ void SFbxSceneOptionWindow::Construct(const FArguments& InArgs)
 				.SlotPadding(2.0f)
 				+ SUniformGridPanel::Slot(0, 0)
 				[
-					IDocumentation::Get()->CreateAnchor(FString("Engine/Content/FBX/ImportOptions"))
+					IDocumentation::Get()->CreateAnchor(FString("fbx-import-options-reference-in-unreal-engine"))
 				]
 				+ SUniformGridPanel::Slot(1, 0)
 				[
