Interchange: Remove unecessary modal dialog confirmation when pressing ESC key

#rnx

diff --git a/Engine/Plugins/Interchange/Editor/Source/Pipelines/Private/SInterchangePipelineConfigurationDialog.cpp b/Engine/Plugins/Interchange/Editor/Source/Pipelines/Private/SInterchangePipelineConfigurationDialog.cpp
--- a/Engine/Plugins/Interchange/Editor/Source/Pipelines/Private/SInterchangePipelineConfigurationDialog.cpp
+++ b/Engine/Plugins/Interchange/Editor/Source/Pipelines/Private/SInterchangePipelineConfigurationDialog.cpp
@@ -668,14 +668,7 @@ FReply SInterchangePipelineConfigurationDialog::OnKeyDown(const FGeometry& MyGeo
 {
 	if (InKeyEvent.GetKey() == EKeys::Escape)
 	{
-		if (!FApp::IsUnattended())
-		{
-			FString Message = FText(LOCTEXT("InterchangePipelineCancelEscKey", "Are you sure you want to cancel the import?")).ToString();
-			if (FPlatformMisc::MessageBoxExt(EAppMsgType::Ok, *Message, TEXT("Cancel Import")) == EAppReturnType::Type::Yes)
-			{
-				return OnCloseDialog(ECloseEventType::Cancel);
-			}
-		}
+		return OnCloseDialog(ECloseEventType::Cancel);
 	}
 	return FReply::Unhandled();
 }
