Fix LOCTEXT duplicate issue
#rnx

diff --git a/Engine/Plugins/Interchange/Editor/Source/Pipelines/Private/InterchangeEditorPipelineDetails.cpp b/Engine/Plugins/Interchange/Editor/Source/Pipelines/Private/InterchangeEditorPipelineDetails.cpp
--- a/Engine/Plugins/Interchange/Editor/Source/Pipelines/Private/InterchangeEditorPipelineDetails.cpp
+++ b/Engine/Plugins/Interchange/Editor/Source/Pipelines/Private/InterchangeEditorPipelineDetails.cpp
@@ -679,11 +679,11 @@ void FInterchangeBaseNodeDetailsCustomization::CustomizeDetails( IDetailLayoutBu
 		if (UInterchangeFactoryBaseNode* FactoryNode = Cast<UInterchangeFactoryBaseNode>(InterchangeBaseNode))
 		{
 			const FString ClassName = FactoryNode->GetObjectClass() ? FactoryNode->GetObjectClass()->GetName() : InterchangeBaseNode->GetClass()->GetName();
-			NodeInformationText = FText::Format(LOCTEXT("NodeFactoryInformationText", "Factory Node ({0})"), FText::FromString(ClassName));
+			NodeInformationText = FText::Format(LOCTEXT("NodeFactoryInformationText_FactoryBaseNode", "Factory Node ({0})"), FText::FromString(ClassName));
 		}
 		else
 		{
-			NodeInformationText = FText::Format(LOCTEXT("NodeFactoryInformationText", "Translated Node ({0})"), FText::FromString(InterchangeBaseNode->GetClass()->GetName()));
+			NodeInformationText = FText::Format(LOCTEXT("NodeFactoryInformationText_BaseNode", "Translated Node ({0})"), FText::FromString(InterchangeBaseNode->GetClass()->GetName()));
 		}
 
 		FDetailWidgetRow& CustomRow = AttributeCategoryBuilder.AddCustomRow(NodeInformationCategoryText)
diff --git a/Engine/Plugins/Interchange/Runtime/Source/Pipelines/Private/InterchangeGenericAssetsPipeline.cpp b/Engine/Plugins/Interchange/Runtime/Source/Pipelines/Private/InterchangeGenericAssetsPipeline.cpp
--- a/Engine/Plugins/Interchange/Runtime/Source/Pipelines/Private/InterchangeGenericAssetsPipeline.cpp
+++ b/Engine/Plugins/Interchange/Runtime/Source/Pipelines/Private/InterchangeGenericAssetsPipeline.cpp
@@ -1451,12 +1451,12 @@ public:
 			FText Tooltip;
 			if (bConflict)
 			{
-				Tooltip = NSLOCTEXT("InterchangeGenericAssetPipeline", "SInterchangeGenericAssetMaterialConflictListRow_Conflict", "Import material is unmatched");
+				Tooltip = NSLOCTEXT("InterchangeGenericAssetPipeline", "SInterchangeGenericAssetMaterialConflictListRow_Conflict_unmatched", "Import material is unmatched");
 			}
 			else
 			{
 				Tooltip = FText::FromString(
-					FText(NSLOCTEXT("InterchangeGenericAssetPipeline", "SInterchangeGenericAssetMaterialConflictListRow_Conflict", "Import material is matched with original asset index: ")).ToString()
+					FText(NSLOCTEXT("InterchangeGenericAssetPipeline", "SInterchangeGenericAssetMaterialConflictListRow_Conflict_matched", "Import material is matched with original asset index: ")).ToString()
 					+ FString::FromInt(Item->bMatched));
 			}
 			return SNew(SBox)
