Interchange: Duplicate pipeline before previewing the result. Pipeline function ScriptedExecutePipeline is not const and the pipeline can behave differently when being called multiple time. Also make sure the material generic pipeline can be called multiple time.

#rnx

diff --git a/Engine/Plugins/Interchange/Editor/Source/Pipelines/Private/SInterchangePipelineConfigurationDialog.cpp b/Engine/Plugins/Interchange/Editor/Source/Pipelines/Private/SInterchangePipelineConfigurationDialog.cpp
--- a/Engine/Plugins/Interchange/Editor/Source/Pipelines/Private/SInterchangePipelineConfigurationDialog.cpp
+++ b/Engine/Plugins/Interchange/Editor/Source/Pipelines/Private/SInterchangePipelineConfigurationDialog.cpp
@@ -884,6 +884,11 @@ void SInterchangePipelineConfigurationDialog::OnFilterOptionsChanged(ECheckBoxSt
 
 FReply SInterchangePipelineConfigurationDialog::OnPreviewImport() const
 {
+	auto ClearObjectFlags = [](UObject* Obj)
+		{
+			Obj->ClearFlags(RF_Standalone | RF_Public);
+			Obj->ClearInternalFlags(EInternalObjectFlags::Async);
+		};
 	UInterchangeBaseNodeContainer* DuplicateBaseNodeContainer = DuplicateObject<UInterchangeBaseNodeContainer>(BaseNodeContainer.Get(), GetTransientPackage());
 
 	TArray<UInterchangeSourceData*> SourceDatas;
@@ -894,8 +899,14 @@ FReply SInterchangePipelineConfigurationDialog::OnPreviewImport() const
 	for (int32 PipelineIndex = 0; PipelineIndex < PipelineListViewItems.Num(); ++PipelineIndex)
 	{
 		const TSharedPtr<FInterchangePipelineItemType> PipelineItem = PipelineListViewItems[PipelineIndex];
-		PipelineItem->Pipeline->SetResultsContainer(Results);
-		PipelineItem->Pipeline->ScriptedExecutePipeline(DuplicateBaseNodeContainer, SourceDatas, FString());
+		
+		//Duplicate the pipeline because ScriptedExecutePipeline is not const
+		if (UInterchangePipelineBase* DuplicatedPipeline = DuplicateObject<UInterchangePipelineBase>(PipelineItem->Pipeline, GetTransientPackage()))
+		{
+			DuplicatedPipeline->SetResultsContainer(Results);
+			DuplicatedPipeline->ScriptedExecutePipeline(DuplicateBaseNodeContainer, SourceDatas, FString());
+			ClearObjectFlags(DuplicatedPipeline);
+		}
 	}
 
 	//Set all node in preview mode so hide the internal data attributes
@@ -920,6 +931,11 @@ FReply SInterchangePipelineConfigurationDialog::OnPreviewImport() const
 
 	FSlateApplication::Get().AddModalWindow(Window, OwnerWindow.Pin(), false);
 
+	//Make sure all temporary object are not flags to persist
+	//We cannot run a gc now since the pipeline we will return are not yet hold by the AsyncHelper, so they will be garbage collect
+	ClearObjectFlags(DuplicateBaseNodeContainer);
+	ClearObjectFlags(Results);
+
 	return FReply::Handled();
 }
 
diff --git a/Engine/Plugins/Interchange/Runtime/Source/Pipelines/Private/InterchangeGenericMaterialPipeline.cpp b/Engine/Plugins/Interchange/Runtime/Source/Pipelines/Private/InterchangeGenericMaterialPipeline.cpp
--- a/Engine/Plugins/Interchange/Runtime/Source/Pipelines/Private/InterchangeGenericMaterialPipeline.cpp
+++ b/Engine/Plugins/Interchange/Runtime/Source/Pipelines/Private/InterchangeGenericMaterialPipeline.cpp
@@ -600,6 +600,9 @@ void UInterchangeGenericMaterialPipeline::ExecutePipeline(UInterchangeBaseNodeCo
 	{
 		return;
 	}
+	
+	MaterialNodes.Empty();
+	MaterialFactoryNodes.Empty();
 
 	//Find all translated node we need for this pipeline
 	BaseNodeContainer->IterateNodes([this](const FString& NodeUid, UInterchangeBaseNode* Node)
