Interchange: Make sure the import option pipeline conflicts are more visible for the user.

#rnx

diff --git a/Engine/Plugins/Interchange/Editor/Source/Pipelines/Private/SInterchangePipelineConfigurationDialog.cpp b/Engine/Plugins/Interchange/Editor/Source/Pipelines/Private/SInterchangePipelineConfigurationDialog.cpp
--- a/Engine/Plugins/Interchange/Editor/Source/Pipelines/Private/SInterchangePipelineConfigurationDialog.cpp
+++ b/Engine/Plugins/Interchange/Editor/Source/Pipelines/Private/SInterchangePipelineConfigurationDialog.cpp
@@ -19,6 +19,7 @@
 #include "SInterchangeGraphInspectorWindow.h"
 #include "SPrimaryButton.h"
 #include "Styling/SlateIconFinder.h"
+#include "Styling/StyleColors.h"
 #include "Widgets/Input/SCheckBox.h"
 #include "Widgets/Input/STextComboBox.h"
 #include "Widgets/Layout/SUniformGridPanel.h"
@@ -66,12 +67,12 @@ void SInterchangePipelineItem::Construct(
 	{
 		FString PipelineNameString = FString::Printf(TEXT("%s (%s)"), *PipelineElement->DisplayName, *PipelineElementPtr->GetClass()->GetName());
 		PipelineName = FText::FromString(PipelineNameString);
-		ConflictInfos = PipelineElementPtr->GetConflictInfos(PipelineElement->ReimportObject, PipelineElement->Container, PipelineElement->SourceData);
-		if (ConflictInfos.Num() > 0)
+		PipelineElement->ConflictInfos = PipelineElementPtr->GetConflictInfos(PipelineElement->ReimportObject, PipelineElement->Container, PipelineElement->SourceData);
+		if (PipelineElement->ConflictInfos.Num() > 0)
 		{
-			ConflictNameList.Reset(ConflictInfos.Num() + 1);
+			ConflictNameList.Reset(PipelineElement->ConflictInfos.Num() + 1);
 			ConflictNameList.Add(ConflictsComboEntry);
-			for (const FInterchangeConflictInfo& ConflictInfo : ConflictInfos)
+			for (const FInterchangeConflictInfo& ConflictInfo : PipelineElement->ConflictInfos)
 			{
 				TSharedPtr<FString> ConflictNamePtr = MakeShared<FString>(ConflictInfo.DisplayName);
 				ConflictNameList.Add(ConflictNamePtr);
@@ -102,6 +103,7 @@ void SInterchangePipelineItem::Construct(
 			[
 				SNew(STextBlock)
 				.Text(PipelineName)
+				.ColorAndOpacity(this, &SInterchangePipelineItem::GetTextColor)
 			]
 			+ SHorizontalBox::Slot()
 			.AutoWidth()
@@ -111,7 +113,7 @@ void SInterchangePipelineItem::Construct(
 				SAssignNew(ConflictComboBox, STextComboBox)
 					.Visibility_Lambda([this]()->EVisibility
 						{
-							return ConflictInfos.Num() > 0 ? EVisibility::All : EVisibility::Collapsed;
+							return PipelineElement->ConflictInfos.Num() > 0 ? EVisibility::All : EVisibility::Collapsed;
 						})
 					.OptionsSource(&ConflictNameList)
 					.OnSelectionChanged_Lambda([this](TSharedPtr<FString> String, ESelectInfo::Type)
@@ -121,7 +123,7 @@ void SInterchangePipelineItem::Construct(
 								return;
 							}
 							//Find and display the conflict info
-							for (const FInterchangeConflictInfo& ConflictInfo : ConflictInfos)
+							for (const FInterchangeConflictInfo& ConflictInfo : PipelineElement->ConflictInfos)
 							{
 								if (ConflictInfo.DisplayName.Equals(*String.Get()))
 								{
@@ -140,7 +142,7 @@ void SInterchangePipelineItem::Construct(
 			]
 		], OwnerTable);
 
-	if (ConflictInfos.Num() > 0)
+	if (PipelineElement->ConflictInfos.Num() > 0)
 	{
 		//Select the conflicts item
 		ConflictComboBox->SetSelectedItem(ConflictsComboEntry);
@@ -160,6 +162,15 @@ const FSlateBrush* SInterchangePipelineItem::GetImageItemIcon() const
 	return TypeIcon;
 }
 
+FSlateColor SInterchangePipelineItem::GetTextColor() const
+{
+	if (PipelineElement->ConflictInfos.Num() > 0)
+	{
+		return FStyleColors::Warning;
+	}
+	return FSlateColor::UseForeground();
+}
+
 /************************************************************************/
 /* SInterchangePipelineConfigurationDialog Implementation                      */
 /************************************************************************/
@@ -514,6 +525,7 @@ void SInterchangePipelineConfigurationDialog::Construct(const FArguments& InArgs
 				.Padding(4.f, 0.f) 
 				[
 					SNew(SPrimaryButton)
+					.Icon(this, &SInterchangePipelineConfigurationDialog::GetImportButtonIcon)
 					.Text(LOCTEXT("InspectorGraphWindow_Import", "Import"))
 					.ToolTipText(this, &SInterchangePipelineConfigurationDialog::GetImportButtonTooltip)
 					.IsEnabled(this, &SInterchangePipelineConfigurationDialog::IsImportButtonEnabled)
@@ -577,6 +589,24 @@ bool SInterchangePipelineConfigurationDialog::IsPropertyVisible(const FPropertyA
 	return true;
 }
 
+const FSlateBrush* SInterchangePipelineConfigurationDialog::GetImportButtonIcon() const
+{
+	const FSlateBrush* TypeIcon = nullptr;
+	for (TSharedPtr<FInterchangePipelineItemType> PipelineItem : PipelineListViewItems)
+	{
+		if (PipelineItem.IsValid() && PipelineItem->Pipeline)
+		{
+			if (PipelineItem->ConflictInfos.Num() > 0)
+			{
+				const FSlateIcon SlateIcon = FSlateIconFinder::FindIcon("Icons.Warning");
+				return SlateIcon.GetOptionalIcon();
+			}
+		}
+	}
+	return TypeIcon;
+}
+
+
 FText SInterchangePipelineConfigurationDialog::GetSourceDescription() const
 {
 	FText ActionDescription;
@@ -672,11 +702,26 @@ bool SInterchangePipelineConfigurationDialog::IsImportButtonEnabled() const
 
 FText SInterchangePipelineConfigurationDialog::GetImportButtonTooltip() const
 {
+	//Pipeline validation
 	TOptional<FText> InvalidReason;
 	if (!ValidateAllPipelineSettings(InvalidReason) && InvalidReason.IsSet())
 	{
 		return InvalidReason.GetValue();
 	}
+
+	//Pipeline conflicts
+	for (TSharedPtr<FInterchangePipelineItemType> PipelineItem : PipelineListViewItems)
+	{
+		if (PipelineItem.IsValid() && PipelineItem->Pipeline)
+		{
+			if (PipelineItem->ConflictInfos.Num() > 0)
+			{
+				return LOCTEXT("ImportButtonConflictTooltip", "There is one or more pipeline conflicts, look at any conflict in the pipeline list to have more detail.");
+			}
+		}
+	}
+
+	//Default tooltip
 	return LOCTEXT("ImportButtonDefaultTooltip", "Selected pipeline stack will be used for the current import");
 }
 
diff --git a/Engine/Plugins/Interchange/Editor/Source/Pipelines/Private/SInterchangePipelineConfigurationDialog.h b/Engine/Plugins/Interchange/Editor/Source/Pipelines/Private/SInterchangePipelineConfigurationDialog.h
--- a/Engine/Plugins/Interchange/Editor/Source/Pipelines/Private/SInterchangePipelineConfigurationDialog.h
+++ b/Engine/Plugins/Interchange/Editor/Source/Pipelines/Private/SInterchangePipelineConfigurationDialog.h
@@ -24,6 +24,7 @@ public:
 	UObject* ReimportObject = nullptr;
 	UInterchangeBaseNodeContainer* Container = nullptr;
 	UInterchangeSourceData* SourceData = nullptr;
+	TArray<FInterchangeConflictInfo> ConflictInfos;
 };
 
 class SInterchangePipelineItem : public STableRow<TSharedPtr<FInterchangePipelineItemType>>
@@ -35,9 +36,9 @@ public:
 		TSharedPtr<FInterchangePipelineItemType> InPipelineElement);
 private:
 	const FSlateBrush* GetImageItemIcon() const;
+	FSlateColor GetTextColor() const;
 
 	TSharedPtr<FInterchangePipelineItemType> PipelineElement = nullptr;
-	TArray<FInterchangeConflictInfo> ConflictInfos;
 	TArray<TSharedPtr<FString>> ConflictNameList;
 	TSharedPtr<FString> ConflictsComboEntry = nullptr;
 	TSharedPtr<STextComboBox> ConflictComboBox = nullptr;
@@ -145,6 +146,8 @@ private:
 
 	void OnFilterOptionsChanged(ECheckBoxState CheckState);
 
+	const FSlateBrush* GetImportButtonIcon() const;
+
 	FReply OnPreviewImport() const;
 
 	TSharedPtr<IDetailsView> PipelineConfigurationDetailsView;
