[Interchange - glTF] Allowing only MaterialInstance creation/import.

#rnx

diff --git a/Engine/Plugins/Interchange/Editor/Source/Pipelines/Private/InterchangeGLTFPipelineCustomizations.cpp b/Engine/Plugins/Interchange/Editor/Source/Pipelines/Private/InterchangeGLTFPipelineCustomizations.cpp
--- a/Engine/Plugins/Interchange/Editor/Source/Pipelines/Private/InterchangeGLTFPipelineCustomizations.cpp
+++ b/Engine/Plugins/Interchange/Editor/Source/Pipelines/Private/InterchangeGLTFPipelineCustomizations.cpp
@@ -33,7 +33,7 @@ void FInterchangeGLTFPipelineCustomization::CustomizeDetails(IDetailLayoutBuilde
 		return;
 	}
 
-	IDetailCategoryBuilder& GLTFCategory = DetailBuilder.EditCategory("glTF");
+	IDetailCategoryBuilder& GLTFCategory = DetailBuilder.EditCategory("GLTF"); //glTF gets auto capitalized and separated into "Gl TF".
 
 	FPropertyEditorModule& PropertyEditorModule = FModuleManager::GetModuleChecked<FPropertyEditorModule>("PropertyEditor");
 
@@ -67,7 +67,7 @@ void FInterchangeGLTFPipelineCustomization::CustomizeDetails(IDetailLayoutBuilde
 			[
 				SNew(STextBlock)
 				.Font(IDetailLayoutBuilder::GetDetailFont())
-				.Text(NSLOCTEXT("InterchangeGLTFPipelineCustomization::Message", "GLTFPredefinedMaterialSubstitutionMessage", "Substitution can be customized in 'Project Settings > Interchange GLTF'."))
+				.Text(NSLOCTEXT("InterchangeGLTFPipelineCustomization::Message", "GLTFPredefinedMaterialSubstitutionMessage", "MaterialInstance Parent Material Substitution can be\ncustomized in 'Project Settings > Interchange GLTF'."))
 				.AutoWrapText(true)
 			];
 	}
diff --git a/Engine/Plugins/Interchange/Runtime/Source/Pipelines/Private/InterchangeglTFPipeline.cpp b/Engine/Plugins/Interchange/Runtime/Source/Pipelines/Private/InterchangeglTFPipeline.cpp
--- a/Engine/Plugins/Interchange/Runtime/Source/Pipelines/Private/InterchangeglTFPipeline.cpp
+++ b/Engine/Plugins/Interchange/Runtime/Source/Pipelines/Private/InterchangeglTFPipeline.cpp
@@ -293,7 +293,7 @@ void UInterchangeGLTFPipeline::ExecutePipeline(UInterchangeBaseNodeContainer* No
 {
 	Super::ExecutePipeline(NodeContainer, InSourceDatas, ContentBasePath);
 
-	if ((FApp::IsGame() || bUseGLTFMaterialInstanceLibrary) && GLTFPipelineSettings)
+	if (GLTFPipelineSettings)
 	{
 		TMap<FString, const UInterchangeShaderGraphNode*> MaterialFactoryNodeUidsToShaderGraphNodes;
 		auto FindGLTFShaderGraphNode = [&MaterialFactoryNodeUidsToShaderGraphNodes, &NodeContainer](const FString& NodeUid, UInterchangeFactoryBaseNode* /*Material or MaterialInstance*/ FactoryNode)
diff --git a/Engine/Plugins/Interchange/Runtime/Source/Pipelines/Public/InterchangeglTFPipeline.h b/Engine/Plugins/Interchange/Runtime/Source/Pipelines/Public/InterchangeglTFPipeline.h
--- a/Engine/Plugins/Interchange/Runtime/Source/Pipelines/Public/InterchangeglTFPipeline.h
+++ b/Engine/Plugins/Interchange/Runtime/Source/Pipelines/Public/InterchangeglTFPipeline.h
@@ -53,9 +53,6 @@ class INTERCHANGEPIPELINES_API UInterchangeGLTFPipeline : public UInterchangePip
 public:
 	TObjectPtr<UGLTFPipelineSettings> GLTFPipelineSettings;
 
-	UPROPERTY(EditAnywhere, BlueprintReadWrite, Category = "glTF", meta = (DisplayName = "Use glTF Material Instance Library"))
-	bool bUseGLTFMaterialInstanceLibrary = true;
-
 protected:
 	virtual void AdjustSettingsForContext(EInterchangePipelineContext ImportType, TObjectPtr<UObject> ReimportAsset) override;
 	virtual void ExecutePipeline(UInterchangeBaseNodeContainer* BaseNodeContainer, const TArray<UInterchangeSourceData*>& SourceDatas, const FString& ContentBasePath) override;
