Interchange: Add convert legacy fbx option to Interchange API.

#rnx

diff --git a/Engine/Plugins/Interchange/Editor/Source/InterchangeEditor/Private/InterchangeFbxAssetImportDataConverter.cpp b/Engine/Plugins/Interchange/Editor/Source/InterchangeEditor/Private/InterchangeFbxAssetImportDataConverter.cpp
--- a/Engine/Plugins/Interchange/Editor/Source/InterchangeEditor/Private/InterchangeFbxAssetImportDataConverter.cpp
+++ b/Engine/Plugins/Interchange/Editor/Source/InterchangeEditor/Private/InterchangeFbxAssetImportDataConverter.cpp
@@ -2,21 +2,26 @@
 #include "InterchangeFbxAssetImportDataConverter.h"
 
 #include "Animation/AnimSequence.h"
+#include "Animation/Skeleton.h"
 #include "EditorFramework/AssetImportData.h"
 #include "Engine/SkeletalMesh.h"
 #include "Engine/SkinnedAssetCommon.h"
 #include "Engine/StaticMesh.h"
 #include "Factories/FbxAnimSequenceImportData.h"
 #include "Factories/FbxAssetImportData.h"
+#include "Factories/FbxImportUI.h"
 #include "Factories/FbxMeshImportData.h"
 #include "Factories/FbxSkeletalMeshImportData.h"
 #include "Factories/FbxStaticMeshImportData.h"
+#include "Factories/FbxTextureImportData.h"
 #include "InterchangeAnimSequenceFactoryNode.h"
 #include "InterchangeAssetImportData.h"
 #include "InterchangeGenericAnimationPipeline.h"
 #include "InterchangeGenericAssetsPipeline.h"
 #include "InterchangeGenericAssetsPipelineSharedSettings.h"
+#include "InterchangeGenericMaterialPipeline.h"
 #include "InterchangeGenericMeshPipeline.h"
+#include "InterchangeGenericTexturePipeline.h"
 #include "InterchangeSkeletalMeshFactoryNode.h"
 #include "InterchangeStaticMeshFactoryNode.h"
 #include "InterchangeManager.h"
@@ -77,6 +82,11 @@ namespace UE::Interchange::Private
 
 	void FillInterchangeGenericAssetsPipelineFromFbxMeshImportData(UInterchangeGenericAssetsPipeline* GenericAssetPipeline, const UFbxMeshImportData* LegacyMeshImportData)
 	{
+		if (!LegacyMeshImportData || !GenericAssetPipeline)
+		{
+			return;
+		}
+
 		GenericAssetPipeline->CommonMeshesProperties->bComputeWeightedNormals = LegacyMeshImportData->bComputeWeightedNormals;
 		GenericAssetPipeline->CommonMeshesProperties->bImportLods = LegacyMeshImportData->bImportMeshLODs;
 		GenericAssetPipeline->CommonMeshesProperties->bBakeMeshes = LegacyMeshImportData->bTransformVertexToAbsolute;
@@ -109,15 +119,17 @@ namespace UE::Interchange::Private
 
 	void FillInterchangeGenericAssetsPipelineFromFbxStaticMeshImportData(UInterchangeGenericAssetsPipeline* GenericAssetPipeline
 		, const UFbxStaticMeshImportData* StaticMeshImportData
-		, UInterchangeBaseNodeContainer* DestinationContainer
-		, const FString& NodeUniqueId
-		, const FString& NodeDisplayLabel)
+		, bool bFillBaseClass = true)
 	{
-		UInterchangeStaticMeshFactoryNode* MeshNode = NewObject<UInterchangeStaticMeshFactoryNode>(DestinationContainer);
-		MeshNode->InitializeStaticMeshNode(NodeUniqueId, NodeDisplayLabel, UStaticMesh::StaticClass()->GetName());
-		DestinationContainer->AddNode(MeshNode);
+		if (!StaticMeshImportData || !GenericAssetPipeline)
+		{
+			return;
+		}
 
-		FillInterchangeGenericAssetsPipelineFromFbxMeshImportData(GenericAssetPipeline, StaticMeshImportData);
+		if (bFillBaseClass)
+		{
+			FillInterchangeGenericAssetsPipelineFromFbxMeshImportData(GenericAssetPipeline, StaticMeshImportData);
+		}
 
 		GenericAssetPipeline->MeshPipeline->bImportCollision = StaticMeshImportData->bAutoGenerateCollision;
 		GenericAssetPipeline->MeshPipeline->bBuildNanite = StaticMeshImportData->bBuildNanite;
@@ -146,15 +158,17 @@ namespace UE::Interchange::Private
 
 	void FillInterchangeGenericAssetsPipelineFromFbxSkeletalMeshImportData(UInterchangeGenericAssetsPipeline* GenericAssetPipeline
 		, const UFbxSkeletalMeshImportData* SkeletalMeshImportData
-		, UInterchangeBaseNodeContainer* DestinationContainer
-		, const FString& NodeUniqueId
-		, const FString& NodeDisplayLabel)
+		, bool bFillBaseClass = true)
 	{
-		UInterchangeSkeletalMeshFactoryNode* MeshNode = NewObject<UInterchangeSkeletalMeshFactoryNode>(DestinationContainer);
-		MeshNode->InitializeSkeletalMeshNode(NodeUniqueId, NodeDisplayLabel, USkeletalMesh::StaticClass()->GetName());
-		DestinationContainer->AddNode(MeshNode);
+		if (!SkeletalMeshImportData || !GenericAssetPipeline)
+		{
+			return;
+		}
 
-		FillInterchangeGenericAssetsPipelineFromFbxMeshImportData(GenericAssetPipeline, SkeletalMeshImportData);
+		if (bFillBaseClass)
+		{
+			FillInterchangeGenericAssetsPipelineFromFbxMeshImportData(GenericAssetPipeline, SkeletalMeshImportData);
+		}
 
 		GenericAssetPipeline->CommonSkeletalMeshesAndAnimationsProperties->bImportMeshesInBoneHierarchy = SkeletalMeshImportData->bImportMeshesInBoneHierarchy;
 		GenericAssetPipeline->CommonMeshesProperties->bKeepSectionsSeparate = SkeletalMeshImportData->bKeepSectionsSeparate;
@@ -209,14 +223,13 @@ namespace UE::Interchange::Private
 	}
 
 	void FillInterchangeGenericAssetsPipelineFromFbxAnimSequenceImportData(UInterchangeGenericAssetsPipeline* GenericAssetPipeline
-		, const UFbxAnimSequenceImportData* AnimSequenceImportData
-		, UInterchangeBaseNodeContainer* DestinationContainer
-		, const FString& NodeUniqueId
-		, const FString& NodeDisplayLabel)
+		, const UFbxAnimSequenceImportData* AnimSequenceImportData)
 	{
-		UInterchangeAnimSequenceFactoryNode* AnimationNode = NewObject<UInterchangeAnimSequenceFactoryNode>(DestinationContainer);
-		AnimationNode->InitializeAnimSequenceNode(NodeUniqueId, NodeDisplayLabel);
-		DestinationContainer->AddNode(AnimationNode);
+		if (!AnimSequenceImportData || !GenericAssetPipeline)
+		{
+			return;
+		}
+
 		switch (AnimSequenceImportData->AnimationLength)
 		{
 		case EFBXAnimationLengthImportType::FBXALIT_ExportedTime:
@@ -484,6 +497,7 @@ namespace UE::Interchange::Private
 
 		//Create a container
 		UInterchangeBaseNodeContainer* DestinationContainer = NewObject<UInterchangeBaseNodeContainer>(DestinationData);
+		DestinationData->SetNodeContainer(DestinationContainer);
 		const FString BasePathToRemove = FPaths::GetBaseFilename(FbxAssetImportData->GetFirstFilename()) + TEXT("_");
 		FString NodeDisplayLabel = Obj->GetName();
 		if (NodeDisplayLabel.StartsWith(BasePathToRemove))
@@ -504,30 +518,35 @@ namespace UE::Interchange::Private
 
 		if (const UFbxStaticMeshImportData* LegacyStaticMeshImportData = Cast<UFbxStaticMeshImportData>(FbxAssetImportData))
 		{
+			UInterchangeStaticMeshFactoryNode* MeshNode = NewObject<UInterchangeStaticMeshFactoryNode>(DestinationContainer);
+			MeshNode->InitializeStaticMeshNode(NodeUniqueId, NodeDisplayLabel, UStaticMesh::StaticClass()->GetName());
+			DestinationContainer->AddNode(MeshNode);
+
+			GenericAssetPipeline->MeshPipeline->CommonMeshesProperties->ForceAllMeshAsType = EInterchangeForceMeshType::IFMT_StaticMesh;
 			check(Obj->IsA<UStaticMesh>());
 			FillInterchangeGenericAssetsPipelineFromFbxStaticMeshImportData(GenericAssetPipeline
-				, LegacyStaticMeshImportData
-				, DestinationContainer
-				, NodeUniqueId
-				, NodeDisplayLabel);
+				, LegacyStaticMeshImportData);
 		}
 		else if (const UFbxSkeletalMeshImportData* LegacySkeletalMeshImportData = Cast<UFbxSkeletalMeshImportData>(FbxAssetImportData))
 		{
+			UInterchangeSkeletalMeshFactoryNode* MeshNode = NewObject<UInterchangeSkeletalMeshFactoryNode>(DestinationContainer);
+			MeshNode->InitializeSkeletalMeshNode(NodeUniqueId, NodeDisplayLabel, USkeletalMesh::StaticClass()->GetName());
+			DestinationContainer->AddNode(MeshNode);
+
+			GenericAssetPipeline->MeshPipeline->CommonMeshesProperties->ForceAllMeshAsType = EInterchangeForceMeshType::IFMT_SkeletalMesh;
 			check(Obj->IsA<USkeletalMesh>());
 			FillInterchangeGenericAssetsPipelineFromFbxSkeletalMeshImportData(GenericAssetPipeline
-				, LegacySkeletalMeshImportData
-				, DestinationContainer
-				, NodeUniqueId
-				, NodeDisplayLabel);
+				, LegacySkeletalMeshImportData);
 		}
 		else if (const UFbxAnimSequenceImportData* LegacyAnimSequenceImportData = Cast<UFbxAnimSequenceImportData>(FbxAssetImportData))
 		{
+			UInterchangeAnimSequenceFactoryNode* AnimationNode = NewObject<UInterchangeAnimSequenceFactoryNode>(DestinationContainer);
+			AnimationNode->InitializeAnimSequenceNode(NodeUniqueId, NodeDisplayLabel);
+			DestinationContainer->AddNode(AnimationNode);
+
 			check(Obj->IsA<UAnimSequence>());
 			FillInterchangeGenericAssetsPipelineFromFbxAnimSequenceImportData(GenericAssetPipeline
-				, LegacyAnimSequenceImportData
-				, DestinationContainer
-				, NodeUniqueId
-				, NodeDisplayLabel);
+				, LegacyAnimSequenceImportData);
 		}
 
 		if (UInterchangeFactoryBaseNode* DestinationFactoryNode = DestinationContainer->GetFactoryNode(NodeUniqueId))
@@ -540,6 +559,147 @@ namespace UE::Interchange::Private
 		return DestinationData;
 	}
 
+	UAssetImportData* ConvertToInterchange(UObject* Owner, const UFbxImportUI* FbxImportUI)
+	{
+		if (!FbxImportUI || !Owner)
+		{
+			return nullptr;
+		}
+		//Create a fbx asset import data and fill the options
+		UInterchangeAssetImportData* DestinationData = NewObject<UInterchangeAssetImportData>(Owner);
+
+		//Create a node container
+		UInterchangeBaseNodeContainer* DestinationContainer = NewObject<UInterchangeBaseNodeContainer>(DestinationData);
+		DestinationData->SetNodeContainer(DestinationContainer);
+
+		TArray<UObject*> Pipelines;
+		UInterchangeGenericAssetsPipeline* GenericAssetPipeline = NewObject<UInterchangeGenericAssetsPipeline>(DestinationData);
+		Pipelines.Add(GenericAssetPipeline);
+		DestinationData->SetPipelines(Pipelines);
+
+		//General Options
+		GenericAssetPipeline->bUseSourceNameForAsset = FbxImportUI->bOverrideFullName;
+		GenericAssetPipeline->ReimportStrategy = EReimportStrategyFlags::ApplyNoProperties;
+
+		//Material Options
+		GenericAssetPipeline->MaterialPipeline->bImportMaterials = FbxImportUI->bImportMaterials;
+		switch (FbxImportUI->TextureImportData->MaterialSearchLocation)
+		{
+			case EMaterialSearchLocation::Local:
+				GenericAssetPipeline->MaterialPipeline->SearchLocation = EInterchangeMaterialSearchLocation::Local;
+				break;
+			case EMaterialSearchLocation::UnderParent:
+				GenericAssetPipeline->MaterialPipeline->SearchLocation = EInterchangeMaterialSearchLocation::UnderParent;
+				break;
+			case EMaterialSearchLocation::UnderRoot:
+				GenericAssetPipeline->MaterialPipeline->SearchLocation = EInterchangeMaterialSearchLocation::UnderRoot;
+				break;
+			case EMaterialSearchLocation::AllAssets:
+				GenericAssetPipeline->MaterialPipeline->SearchLocation = EInterchangeMaterialSearchLocation::AllAssets;
+				break;
+			case EMaterialSearchLocation::DoNotSearch:
+				GenericAssetPipeline->MaterialPipeline->SearchLocation = EInterchangeMaterialSearchLocation::DoNotSearch;
+				break;
+		}
+		if (FbxImportUI->TextureImportData->BaseMaterialName.IsAsset())
+		{
+			GenericAssetPipeline->MaterialPipeline->MaterialImport = EInterchangeMaterialImportOption::ImportAsMaterialInstances;
+			GenericAssetPipeline->MaterialPipeline->ParentMaterial = FbxImportUI->TextureImportData->BaseMaterialName;
+		}
+		else
+		{
+			GenericAssetPipeline->MaterialPipeline->MaterialImport = EInterchangeMaterialImportOption::ImportAsMaterials;
+			GenericAssetPipeline->MaterialPipeline->ParentMaterial.Reset();
+		}
+
+		//Texture Options
+		GenericAssetPipeline->MaterialPipeline->TexturePipeline->bImportTextures = FbxImportUI->bImportTextures;
+		GenericAssetPipeline->MaterialPipeline->TexturePipeline->bFlipNormalMapGreenChannel = FbxImportUI->TextureImportData->bInvertNormalMaps;
+
+		//Default the force animation to false
+		GenericAssetPipeline->CommonSkeletalMeshesAndAnimationsProperties->bImportOnlyAnimations = false;
+
+		//Discover if we must import something in particular
+		if (FbxImportUI->MeshTypeToImport == EFBXImportType::FBXIT_SkeletalMesh
+			|| FbxImportUI->bImportAsSkeletal)
+		{
+			GenericAssetPipeline->MeshPipeline->bImportSkeletalMeshes = true;
+			GenericAssetPipeline->MeshPipeline->CommonMeshesProperties->ForceAllMeshAsType = EInterchangeForceMeshType::IFMT_SkeletalMesh;
+
+			GenericAssetPipeline->AnimationPipeline->bImportAnimations = FbxImportUI->bImportAnimations;
+
+			GenericAssetPipeline->ImportOffsetRotation = FbxImportUI->SkeletalMeshImportData->ImportRotation;
+			GenericAssetPipeline->ImportOffsetTranslation = FbxImportUI->SkeletalMeshImportData->ImportTranslation;
+			GenericAssetPipeline->ImportOffsetUniformScale = FbxImportUI->SkeletalMeshImportData->ImportUniformScale;
+
+			FillInterchangeGenericAssetsPipelineFromFbxMeshImportData(GenericAssetPipeline, Cast<UFbxSkeletalMeshImportData>(FbxImportUI->SkeletalMeshImportData));
+		}
+		else if (FbxImportUI->MeshTypeToImport == EFBXImportType::FBXIT_StaticMesh)
+		{
+			GenericAssetPipeline->MeshPipeline->bImportStaticMeshes = true;
+			GenericAssetPipeline->MeshPipeline->CommonMeshesProperties->ForceAllMeshAsType = EInterchangeForceMeshType::IFMT_StaticMesh;
+
+			GenericAssetPipeline->AnimationPipeline->bImportAnimations = false;
+			GenericAssetPipeline->MeshPipeline->bImportSkeletalMeshes = false;
+
+			GenericAssetPipeline->ImportOffsetRotation = FbxImportUI->StaticMeshImportData->ImportRotation;
+			GenericAssetPipeline->ImportOffsetTranslation = FbxImportUI->StaticMeshImportData->ImportTranslation;
+			GenericAssetPipeline->ImportOffsetUniformScale = FbxImportUI->StaticMeshImportData->ImportUniformScale;
+
+			FillInterchangeGenericAssetsPipelineFromFbxMeshImportData(GenericAssetPipeline, Cast<UFbxStaticMeshImportData>(FbxImportUI->StaticMeshImportData));
+		}
+		else if (FbxImportUI->MeshTypeToImport == EFBXImportType::FBXIT_Animation)
+		{
+			GenericAssetPipeline->AnimationPipeline->bImportAnimations = true;
+			if (FbxImportUI->Skeleton)
+			{
+				GenericAssetPipeline->CommonSkeletalMeshesAndAnimationsProperties->bImportOnlyAnimations = true;
+				GenericAssetPipeline->CommonSkeletalMeshesAndAnimationsProperties->Skeleton = FbxImportUI->Skeleton;
+
+				GenericAssetPipeline->MeshPipeline->bImportStaticMeshes = false;
+				GenericAssetPipeline->MeshPipeline->bImportSkeletalMeshes = false;
+			}
+			else
+			{
+				GenericAssetPipeline->MeshPipeline->bImportStaticMeshes = true;
+				GenericAssetPipeline->MeshPipeline->bImportSkeletalMeshes = true;
+			}
+
+			GenericAssetPipeline->ImportOffsetRotation = FbxImportUI->AnimSequenceImportData->ImportRotation;
+			GenericAssetPipeline->ImportOffsetTranslation = FbxImportUI->AnimSequenceImportData->ImportTranslation;
+			GenericAssetPipeline->ImportOffsetUniformScale = FbxImportUI->AnimSequenceImportData->ImportUniformScale;
+		}
+		else
+		{
+			//Allow importing all type
+			GenericAssetPipeline->MeshPipeline->CommonMeshesProperties->ForceAllMeshAsType = EInterchangeForceMeshType::IFMT_None;
+			GenericAssetPipeline->MeshPipeline->bImportStaticMeshes = true;
+			GenericAssetPipeline->MeshPipeline->bImportSkeletalMeshes = true;
+			GenericAssetPipeline->AnimationPipeline->bImportAnimations = true;
+			//Use the static mesh data
+			FillInterchangeGenericAssetsPipelineFromFbxMeshImportData(GenericAssetPipeline, Cast<UFbxStaticMeshImportData>(FbxImportUI->StaticMeshImportData));
+		}
+
+		if (const UFbxStaticMeshImportData* LegacyStaticMeshImportData = Cast<UFbxStaticMeshImportData>(FbxImportUI->StaticMeshImportData))
+		{
+			FillInterchangeGenericAssetsPipelineFromFbxStaticMeshImportData(GenericAssetPipeline
+				, LegacyStaticMeshImportData
+				, false);
+		}
+		if (const UFbxSkeletalMeshImportData* LegacySkeletalMeshImportData = Cast<UFbxSkeletalMeshImportData>(FbxImportUI->SkeletalMeshImportData))
+		{
+			FillInterchangeGenericAssetsPipelineFromFbxSkeletalMeshImportData(GenericAssetPipeline
+				, LegacySkeletalMeshImportData
+				, false);
+		}
+		if (const UFbxAnimSequenceImportData* LegacyAnimSequenceImportData = Cast<UFbxAnimSequenceImportData>(FbxImportUI->AnimSequenceImportData))
+		{
+			FillInterchangeGenericAssetsPipelineFromFbxAnimSequenceImportData(GenericAssetPipeline
+				, LegacyAnimSequenceImportData);
+		}
+		return DestinationData;
+	}
+
 	UAssetImportData* ConvertData(UObject* Obj, UAssetImportData* SourceData, const bool bInterchangeSupportTargetExtension)
 	{
 		if (const UInterchangeAssetImportData* InterchangeSourceData = Cast<UInterchangeAssetImportData>(SourceData))
@@ -662,4 +822,26 @@ bool UInterchangeFbxAssetImportDataConverter::ConvertImportData(UObject* Obj, co
 	return bResult;
 }
 
+bool UInterchangeFbxAssetImportDataConverter::ConvertImportData(const UObject* SourceImportData, UObject** DestinationImportData) const
+{
+	bool bResult = false;
+	if (!SourceImportData || !DestinationImportData)
+	{
+		return bResult;
+	}
+
+	UInterchangeManager& InterchangeManager = UInterchangeManager::GetInterchangeManager();
+	if (SourceImportData->IsA<UFbxImportUI>())
+	{
+		//Convert Legacy Fbx to Interchange
+		*DestinationImportData = UE::Interchange::Private::ConvertToInterchange(GetTransientPackage(), Cast<UFbxImportUI>(SourceImportData));
+		bResult = true;
+	}
+	else if (SourceImportData->IsA<UInterchangeAssetImportData>())
+	{
+		//TODO Convert Interchange to Legacy Fbx
+		//*DestinationImportData = UE::Interchange::Private::ConvertToLegacyFbx(GetTransientPackage(), Cast<UInterchangeAssetImportData>(SourceImportData));
+	}
+	return bResult;
+}
 
diff --git a/Engine/Plugins/Interchange/Editor/Source/InterchangeEditor/Public/InterchangeFbxAssetImportDataConverter.h b/Engine/Plugins/Interchange/Editor/Source/InterchangeEditor/Public/InterchangeFbxAssetImportDataConverter.h
--- a/Engine/Plugins/Interchange/Editor/Source/InterchangeEditor/Public/InterchangeFbxAssetImportDataConverter.h
+++ b/Engine/Plugins/Interchange/Editor/Source/InterchangeEditor/Public/InterchangeFbxAssetImportDataConverter.h
@@ -16,4 +16,5 @@ class INTERCHANGEEDITOR_API UInterchangeFbxAssetImportDataConverter : public UIn
 
 public:
 	virtual bool ConvertImportData(UObject* Object, const FString& ToExtension) const override;
+	virtual bool ConvertImportData(const UObject* SourceImportData, UObject** DestinationImportDataClass) const override;
 };
diff --git a/Engine/Source/Developer/AssetTools/Private/AssetTools.cpp b/Engine/Source/Developer/AssetTools/Private/AssetTools.cpp
--- a/Engine/Source/Developer/AssetTools/Private/AssetTools.cpp
+++ b/Engine/Source/Developer/AssetTools/Private/AssetTools.cpp
@@ -3588,6 +3588,10 @@ TArray<UObject*> UAssetToolsImpl::ImportAssetsInternal(const TArray<FString>& Fi
 								ImportAssetParameters.OverridePipelines.Add(OverridePipelinePath);
 							}
 						}
+						else
+						{
+							InterchangeManager.ConvertImportData(Params.AssetImportTask->Options, ImportAssetParameters);
+						}
 					}
 					UE::Interchange::FAssetImportResultRef InterchangeResult = (InterchangeManager.ImportAssetAsync(DestinationPath, ScopedSourceData.GetSourceData(), ImportAssetParameters));
 
diff --git a/Engine/Source/Runtime/Interchange/Engine/Private/InterchangeManager.cpp b/Engine/Source/Runtime/Interchange/Engine/Private/InterchangeManager.cpp
--- a/Engine/Source/Runtime/Interchange/Engine/Private/InterchangeManager.cpp
+++ b/Engine/Source/Runtime/Interchange/Engine/Private/InterchangeManager.cpp
@@ -1033,7 +1033,7 @@ bool UInterchangeManager::RegisterImportDataConverter(const UClass* Converter)
 	return true;
 }
 
-bool UInterchangeManager::ConvertImportData(UObject* Object, const FString& Extension)
+bool UInterchangeManager::ConvertImportData(UObject* Object, const FString& Extension) const
 {
 	for (TPair<TObjectPtr<const UClass>, TObjectPtr<UInterchangeAssetImportDataConverterBase>> RegisteredConverter : RegisteredConverters)
 	{
@@ -1045,6 +1045,71 @@ bool UInterchangeManager::ConvertImportData(UObject* Object, const FString& Exte
 	return false;
 }
 
+bool UInterchangeManager::ConvertImportData(const UObject* SourceImportData, FImportAssetParameters& ImportAssetParameters) const
+{
+	UObject* DestinationImportData = nullptr;
+	for (TPair<TObjectPtr<const UClass>, TObjectPtr<UInterchangeAssetImportDataConverterBase>> RegisteredConverter : RegisteredConverters)
+	{
+		if (!RegisteredConverter.Value->ConvertImportData(SourceImportData, &DestinationImportData))
+		{
+			return false;
+		}
+	}
+	if (UInterchangeAssetImportData* AssetImportData = Cast<UInterchangeAssetImportData>(DestinationImportData))
+	{
+		//We can use the default pipeline stack, if it contain a pipeline that match the converted pipeline class
+		bool bUseDefaultPipelineStack = false;
+		TArray<UInterchangePipelineBase*> DuplicateDefaultPipelines;
+		//Get the Interchange Default stack
+		constexpr bool bIsSceneImport = false;
+		const FInterchangeImportSettings& InterchangeImportSettings = FInterchangeProjectSettingsUtils::GetDefaultImportSettings(bIsSceneImport);
+		//Verify if we can use the default stack or not
+		if (AssetImportData->GetNumberOfPipelines() == 1
+			&& InterchangeImportSettings.PipelineStacks.Contains(InterchangeImportSettings.DefaultPipelineStack))
+		{
+			if (UInterchangePipelineBase* ConvertedPipeline = Cast<UInterchangePipelineBase>(AssetImportData->GetPipelines()[0]))
+			{
+				UClass* ConvertedPipelineClass = ConvertedPipeline->GetClass();
+				const FInterchangePipelineStack& PipelineStack = InterchangeImportSettings.PipelineStacks.FindChecked(InterchangeImportSettings.DefaultPipelineStack);
+				for (const FSoftObjectPath& PipelinePath : PipelineStack.Pipelines)
+				{
+					if (UInterchangePipelineBase* GeneratedPipeline = UE::Interchange::GeneratePipelineInstance(PipelinePath, GetTransientPackage()))
+					{
+						GeneratedPipeline->AdjustSettingsForContext(EInterchangePipelineContext::AssetImport, nullptr);
+						if (GeneratedPipeline->IsA(ConvertedPipelineClass))
+						{
+							//We found a match, so we will use the default pipeline stacks
+							bUseDefaultPipelineStack = true;
+							DuplicateDefaultPipelines.Add(ConvertedPipeline);
+						}
+						else
+						{
+							DuplicateDefaultPipelines.Add(GeneratedPipeline);
+						}
+					}
+				}
+			}
+		}
+
+		if (bUseDefaultPipelineStack)
+		{
+			for (UInterchangePipelineBase* Pipeline : DuplicateDefaultPipelines)
+			{
+				ImportAssetParameters.OverridePipelines.Add(Pipeline);
+			}
+		}
+		else
+		{
+			for (UObject* Pipeline : AssetImportData->GetPipelines())
+			{
+				ImportAssetParameters.OverridePipelines.Add(Pipeline);
+			}
+		}
+		return true;
+	}
+	return false;
+}
+
 TArray<FString> UInterchangeManager::GetSupportedFormats(const EInterchangeTranslatorType ForTranslatorType) const
 {
 	TArray<FString> FileExtensions;
diff --git a/Engine/Source/Runtime/Interchange/Engine/Public/InterchangeAssetImportData.h b/Engine/Source/Runtime/Interchange/Engine/Public/InterchangeAssetImportData.h
--- a/Engine/Source/Runtime/Interchange/Engine/Public/InterchangeAssetImportData.h
+++ b/Engine/Source/Runtime/Interchange/Engine/Public/InterchangeAssetImportData.h
@@ -197,4 +197,15 @@ public:
 	{
 		return false;
 	}
+
+	/**
+	 * Convert the asset import data from the source to the destination.
+	 * The function should return true only if it has convert the asset import data, false otherwise.
+	 *
+	 * The system will call all object deriving from this class until one convert the data.
+	 */
+	virtual bool ConvertImportData(const UObject* SourceImportData, UObject** DestinationImportDataClass) const
+	{
+		return false;
+	}
 };
diff --git a/Engine/Source/Runtime/Interchange/Engine/Public/InterchangeManager.h b/Engine/Source/Runtime/Interchange/Engine/Public/InterchangeManager.h
--- a/Engine/Source/Runtime/Interchange/Engine/Public/InterchangeManager.h
+++ b/Engine/Source/Runtime/Interchange/Engine/Public/InterchangeManager.h
@@ -493,7 +493,15 @@ public:
 	 * @Param Extension - The file extension we want to import
 	 * @return true if one of the converter has convert the data. False otherwise.
 	 */
-	INTERCHANGEENGINE_API bool ConvertImportData(UObject* Object, const FString& Extension);
+	INTERCHANGEENGINE_API bool ConvertImportData(UObject* Object, const FString& Extension) const;
+
+	/**
+	 * Call all the registered converter, if one converter want
+	 * @Param SourceImportData - The source import data options.
+	 * @Param DestinationImportData - The destination import data options.
+	 * @return true if one of the converter has convert the data. False otherwise.
+	 */
+	INTERCHANGEENGINE_API bool ConvertImportData(const UObject* SourceImportData, FImportAssetParameters& ImportAssetParameters) const;
 
 	/**
 	 * Returns the list of supported formats for a given translator type.
