Interchange: Create "Advanced/Basic" pipeline visualization dialog options. This allow to easily have a simplified view of options when users are importing.

#rnx

#virtualized

diff --git a/Engine/Plugins/Interchange/Editor/Source/Pipelines/Private/InterchangeEditorPipelineDetails.cpp b/Engine/Plugins/Interchange/Editor/Source/Pipelines/Private/InterchangeEditorPipelineDetails.cpp
--- a/Engine/Plugins/Interchange/Editor/Source/Pipelines/Private/InterchangeEditorPipelineDetails.cpp
+++ b/Engine/Plugins/Interchange/Editor/Source/Pipelines/Private/InterchangeEditorPipelineDetails.cpp
@@ -271,6 +271,7 @@ void FInterchangePipelineBaseDetailsCustomization::CustomizeDetails(IDetailLayou
 
 	const bool bAllowPropertyStatesEdition = InterchangePipeline->CanEditPropertiesStates();
 	const bool bIsReimportContext = InterchangePipeline->IsReimportContext();
+	const bool bISBasicLayout = InterchangePipeline->IsBasicLayout();
 
 	TArray<FName> AllCategoryNames;
 	CachedDetailBuilder->GetCategoryNames(AllCategoryNames);
@@ -376,7 +377,7 @@ void FInterchangePipelineBaseDetailsCustomization::CustomizeDetails(IDetailLayou
 				bool IsLocked = false;
 				if (const FInterchangePipelinePropertyStates* PropertyStates = InterchangePipeline->GetPropertyStates(PropertyPath))
 				{
-					if (!PropertyStates->IsPropertyVisible(bIsReimportContext))
+					if (!PropertyStates->IsPropertyVisible(bIsReimportContext, bISBasicLayout))
 					{
 						continue;
 					}
@@ -486,6 +487,51 @@ void FInterchangePipelineBaseDetailsCustomization::CustomizeDetails(IDetailLayou
 					.AutoWidth()
 					.Padding(6.0f, 1.0f, 3.0f, 1.0f)
 					.VAlign(VAlign_Center)
+					[
+						SNew(STextBlock)
+						.Font(IDetailLayoutBuilder::GetDetailFont())
+						.Text(NSLOCTEXT("InterchangePipelineBaseDetails::CustomizeDetails", "ShowWhenBasicLayoutText", "Basic Layout"))
+					]
+					+ SHorizontalBox::Slot()
+					.AutoWidth()
+					.Padding(3.0f, 1.0f)
+					[
+						SNew(SCheckBox)
+						.CheckedImage(FAppStyle::Get().GetBrush("Icons.Hidden"))
+						.CheckedHoveredImage(FAppStyle::Get().GetBrush("Icons.Hidden"))
+						.CheckedPressedImage(FAppStyle::Get().GetBrush("Icons.Hidden"))
+						.UncheckedImage(FAppStyle::Get().GetBrush("Icons.Visible"))
+						.UncheckedHoveredImage(FAppStyle::Get().GetBrush("Icons.Visible"))
+						.UncheckedPressedImage(FAppStyle::Get().GetBrush("Icons.Visible"))
+						.ToolTipText(NSLOCTEXT("InterchangePipelineBaseDetails::CustomizeDetails", "VisibleTooltipBasicLayout", "If true this property will be visible when displaying the interchange import dialog with basic layout."))
+						.OnCheckStateChanged_Lambda([InterchangePipelinePtr = InterchangePipeline, PropertyPath](ECheckBoxState CheckType)
+						{
+							if (!ensure(InterchangePipelinePtr.IsValid()))
+							{
+								return;
+							}
+							FScopedTransaction ScopedTransaction(NSLOCTEXT("InterchangePipelineBaseDetails::CustomizeDetails", "TransactionvisibilityPropertiesToggleImport", "Toggle property visibility at import"), !GIsTransacting);
+							InterchangePipelinePtr->Modify();
+							InterchangePipelinePtr->FindOrAddPropertyStates(PropertyPath).SetPropertyBasicLayoutVisibility((CheckType != ECheckBoxState::Checked));
+							InterchangePipelinePtr->PostEditChange();
+						})
+						.IsChecked_Lambda([InterchangePipelinePtr = InterchangePipeline, PropertyPath]()
+						{
+							if (!InterchangePipelinePtr.IsValid())
+							{
+								return ECheckBoxState::Unchecked;
+							}
+							if (const FInterchangePipelinePropertyStates* PropertyStates = InterchangePipelinePtr->GetPropertyStates(PropertyPath))
+							{
+								return PropertyStates->IsPropertyVisibleInBasicLayout() ? ECheckBoxState::Unchecked : ECheckBoxState::Checked;
+							}
+							return ECheckBoxState::Unchecked;
+						})
+					]
+					+ SHorizontalBox::Slot()
+					.AutoWidth()
+					.Padding(6.0f, 1.0f, 3.0f, 1.0f)
+					.VAlign(VAlign_Center)
 					[
 						SNew(STextBlock)
 						.Font(IDetailLayoutBuilder::GetDetailFont())
@@ -505,26 +551,26 @@ void FInterchangePipelineBaseDetailsCustomization::CustomizeDetails(IDetailLayou
 						.ToolTipText(NSLOCTEXT("InterchangePipelineBaseDetails::CustomizeDetails", "VisibleTooltipImport", "If true this property will be visible when displaying the interchange import dialog."))
 						.OnCheckStateChanged_Lambda([InterchangePipelinePtr = InterchangePipeline, PropertyPath](ECheckBoxState CheckType)
 						{
-							constexpr bool bIsReimportContextLocal = false;
 							if (!ensure(InterchangePipelinePtr.IsValid()))
 							{
 								return;
 							}
 							FScopedTransaction ScopedTransaction(NSLOCTEXT("InterchangePipelineBaseDetails::CustomizeDetails", "TransactionvisibilityPropertiesToggleImport", "Toggle property visibility at import"), !GIsTransacting);
 							InterchangePipelinePtr->Modify();
-							InterchangePipelinePtr->FindOrAddPropertyStates(PropertyPath).SetPropertyVisible(bIsReimportContextLocal, (CheckType != ECheckBoxState::Checked));
+							InterchangePipelinePtr->FindOrAddPropertyStates(PropertyPath).SetPropertyImportVisibility((CheckType != ECheckBoxState::Checked));
 							InterchangePipelinePtr->PostEditChange();
 						})
 						.IsChecked_Lambda([InterchangePipelinePtr = InterchangePipeline, PropertyPath]()
 						{
 							constexpr bool bIsReimportContextLocal = false;
+							constexpr bool bIsBasicLayoutLocal = false;
 							if (!InterchangePipelinePtr.IsValid())
 							{
 								return ECheckBoxState::Unchecked;
 							}
 							if (const FInterchangePipelinePropertyStates* PropertyStates = InterchangePipelinePtr->GetPropertyStates(PropertyPath))
 							{
-								return PropertyStates->IsPropertyVisible(bIsReimportContextLocal) ? ECheckBoxState::Unchecked : ECheckBoxState::Checked;
+								return PropertyStates->IsPropertyVisible(bIsReimportContextLocal, bIsBasicLayoutLocal) ? ECheckBoxState::Unchecked : ECheckBoxState::Checked;
 							}
 							return ECheckBoxState::Unchecked;
 						})
@@ -552,26 +598,26 @@ void FInterchangePipelineBaseDetailsCustomization::CustomizeDetails(IDetailLayou
 						.ToolTipText(NSLOCTEXT("InterchangePipelineBaseDetails::CustomizeDetails", "VisibleTooltipReimport", "If true this property will be visible when displaying the interchange reimport dialog."))
 						.OnCheckStateChanged_Lambda([InterchangePipelinePtr = InterchangePipeline, PropertyPath](ECheckBoxState CheckType)
 						{
-							constexpr bool bIsReimportContextLocal = true;
 							if (!ensure(InterchangePipelinePtr.IsValid()))
 							{
 								return;
 							}
 							FScopedTransaction ScopedTransaction(NSLOCTEXT("InterchangePipelineBaseDetails::CustomizeDetails", "TransactionvisibilityPropertiesToggleReimport", "Toggle property visibility at reimport"), !GIsTransacting);
 							InterchangePipelinePtr->Modify();
-							InterchangePipelinePtr->FindOrAddPropertyStates(PropertyPath).SetPropertyVisible(bIsReimportContextLocal, (CheckType != ECheckBoxState::Checked));
+							InterchangePipelinePtr->FindOrAddPropertyStates(PropertyPath).SetPropertyReimportVisibility((CheckType != ECheckBoxState::Checked));
 							InterchangePipelinePtr->PostEditChange();
 						})
 						.IsChecked_Lambda([InterchangePipelinePtr = InterchangePipeline, PropertyPath]()
 						{
 							constexpr bool bIsReimportContextLocal = true;
+							constexpr bool bIsBasicLayoutLocal = false;
 							if (!InterchangePipelinePtr.IsValid())
 							{
 								return ECheckBoxState::Unchecked;
 							}
 							if (const FInterchangePipelinePropertyStates* PropertyStates = InterchangePipelinePtr->GetPropertyStates(PropertyPath))
 							{
-								return PropertyStates->IsPropertyVisible(bIsReimportContextLocal) ? ECheckBoxState::Unchecked : ECheckBoxState::Checked;
+								return PropertyStates->IsPropertyVisible(bIsReimportContextLocal, bIsBasicLayoutLocal) ? ECheckBoxState::Unchecked : ECheckBoxState::Checked;
 							}
 							return ECheckBoxState::Unchecked;
 						})
diff --git a/Engine/Plugins/Interchange/Editor/Source/Pipelines/Private/SInterchangePipelineConfigurationDialog.cpp b/Engine/Plugins/Interchange/Editor/Source/Pipelines/Private/SInterchangePipelineConfigurationDialog.cpp
--- a/Engine/Plugins/Interchange/Editor/Source/Pipelines/Private/SInterchangePipelineConfigurationDialog.cpp
+++ b/Engine/Plugins/Interchange/Editor/Source/Pipelines/Private/SInterchangePipelineConfigurationDialog.cpp
@@ -245,6 +245,7 @@ TSharedRef<SBox> SInterchangePipelineConfigurationDialog::SpawnPipelineConfigura
 						GeneratedPipeline->LoadSettings(Stack.StackName);
 						GeneratedPipeline->PreDialogCleanup(Stack.StackName);
 					}
+					GeneratedPipeline->SetBasicLayoutMode(bBasicLayout);
 					if (bFilterOptions && BaseNodeContainer.IsValid())
 					{
 						GeneratedPipeline->FilterPropertiesFromTranslatedData(BaseNodeContainer.Get());
@@ -407,6 +408,7 @@ void SInterchangePipelineConfigurationDialog::Construct(const FArguments& InArgs
 	else if(GConfig->DoesSectionExist(TEXT("InterchangeImportDialogOptions"), GEditorPerProjectIni))
 	{
 		GConfig->GetBool(TEXT("InterchangeImportDialogOptions"), TEXT("FilterOptions"), bFilterOptions, GEditorPerProjectIni);
+		GConfig->GetBool(TEXT("InterchangeImportDialogOptions"), TEXT("BasicLayout"), bBasicLayout, GEditorPerProjectIni);
 	}
 
 	this->ChildSlot
@@ -433,6 +435,30 @@ void SInterchangePipelineConfigurationDialog::Construct(const FArguments& InArgs
 				.AutoWidth()
 				.HAlign(HAlign_Right)
 				.Padding(10.0f, 2.0f, 0.0f, 2.0f)
+				[
+					SNew(SHorizontalBox)
+					.ToolTipText(LOCTEXT("SInterchangePipelineConfigurationDialog_BasicLayoutOptions_tooltip", "Basic Layout display only the basic pipelines properties."))
+					+ SHorizontalBox::Slot()
+					.AutoWidth()
+					.HAlign(HAlign_Right)
+					.VAlign(VAlign_Center)
+					.Padding(4.f, 0.f)
+					[
+						SNew(STextBlock)
+						.Text(LOCTEXT("SInterchangePipelineConfigurationDialog_BasicLayoutOptions", "Basic Layout"))
+					]
+					+ SHorizontalBox::Slot()
+					.Padding(4.f, 0.f)
+					[
+						SNew(SCheckBox)
+						.IsChecked(this, &SInterchangePipelineConfigurationDialog::IsBasicLayoutEnabled)
+						.OnCheckStateChanged(this, &SInterchangePipelineConfigurationDialog::OnBasicLayoutChanged)
+					]
+				]
+				+SHorizontalBox::Slot()
+				.AutoWidth()
+				.HAlign(HAlign_Right)
+				.Padding(10.0f, 2.0f, 0.0f, 2.0f)
 				[
 					SNew(SHorizontalBox)
 					.ToolTipText(LOCTEXT("SInterchangePipelineConfigurationDialog_FilterPipelineOptions_tooltip", "Filter the pipeline options using the source content data."))
@@ -663,6 +689,7 @@ FReply SInterchangePipelineConfigurationDialog::OnResetToDefault()
 								if (UInterchangePipelineBase* GeneratedPipeline = UE::Interchange::GeneratePipelineInstanceInSourceAssetPackage(DefaultPipeline))
 								{
 									GeneratedPipeline->TransferAdjustSettings(DefaultPipeline);
+									GeneratedPipeline->SetBasicLayoutMode(bBasicLayout);
 									if(bFilterOptions && BaseNodeContainer.IsValid())
 									{
 										GeneratedPipeline->FilterPropertiesFromTranslatedData(BaseNodeContainer.Get());
@@ -846,6 +873,7 @@ void SInterchangePipelineConfigurationDialog::RefreshStack(bool bStackSelectionC
 						GeneratedPipeline->PreDialogCleanup(Stack.StackName);
 					}
 				}
+				GeneratedPipeline->SetBasicLayoutMode(bBasicLayout);
 				if (bFilterOptions && BaseNodeContainer.IsValid())
 				{
 					GeneratedPipeline->FilterPropertiesFromTranslatedData(BaseNodeContainer.Get());
@@ -932,6 +960,22 @@ void SInterchangePipelineConfigurationDialog::OnFilterOptionsChanged(ECheckBoxSt
 	GConfig->SetBool(TEXT("InterchangeImportDialogOptions"), TEXT("FilterOptions"), bFilterOptions, GEditorPerProjectIni);
 }
 
+void SInterchangePipelineConfigurationDialog::OnBasicLayoutChanged(ECheckBoxState CheckState)
+{
+	bool bNewCheckValue = CheckState == ECheckBoxState::Checked ? true : false;
+	if (bNewCheckValue == bBasicLayout)
+	{
+		//Check state did not change
+		return;
+	}
+	bBasicLayout = bNewCheckValue;
+	//Refresh the pipeline
+	constexpr bool bStackSelectionChange = false;
+	RefreshStack(bStackSelectionChange);
+
+	GConfig->SetBool(TEXT("InterchangeImportDialogOptions"), TEXT("BasicLayout"), bBasicLayout, GEditorPerProjectIni);
+}
+
 FReply SInterchangePipelineConfigurationDialog::OnPreviewImport() const
 {
 	auto ClearObjectFlags = [](UObject* Obj)
diff --git a/Engine/Plugins/Interchange/Editor/Source/Pipelines/Private/SInterchangePipelineConfigurationDialog.h b/Engine/Plugins/Interchange/Editor/Source/Pipelines/Private/SInterchangePipelineConfigurationDialog.h
--- a/Engine/Plugins/Interchange/Editor/Source/Pipelines/Private/SInterchangePipelineConfigurationDialog.h
+++ b/Engine/Plugins/Interchange/Editor/Source/Pipelines/Private/SInterchangePipelineConfigurationDialog.h
@@ -144,7 +144,13 @@ private:
 		return bFilterOptions ? ECheckBoxState::Checked : ECheckBoxState::Unchecked;
 	}
 
+	ECheckBoxState IsBasicLayoutEnabled() const
+	{
+		return bBasicLayout ? ECheckBoxState::Checked : ECheckBoxState::Unchecked;
+	}
+
 	void OnFilterOptionsChanged(ECheckBoxState CheckState);
+	void OnBasicLayoutChanged(ECheckBoxState CheckState);
 
 	const FSlateBrush* GetImportButtonIcon() const;
 
@@ -159,6 +165,7 @@ private:
 	bool bImportAll = false;
 
 	bool bFilterOptions = false;
+	bool bBasicLayout = false;
 
 	FName CurrentStackName = NAME_None;
 	TObjectPtr<UInterchangePipelineBase> CurrentSelectedPipeline = nullptr;
diff --git a/Engine/Source/Runtime/Interchange/Core/Private/InterchangePipelineBase.cpp b/Engine/Source/Runtime/Interchange/Core/Private/InterchangePipelineBase.cpp
--- a/Engine/Source/Runtime/Interchange/Core/Private/InterchangePipelineBase.cpp
+++ b/Engine/Source/Runtime/Interchange/Core/Private/InterchangePipelineBase.cpp
@@ -69,6 +69,7 @@ void UInterchangePipelineBase::TransferAdjustSettings(UInterchangePipelineBase*
 	CachePropertiesStates = SourcePipeline->CachePropertiesStates;
 	bAllowPropertyStatesEdition = SourcePipeline->bAllowPropertyStatesEdition;
 	bIsReimportContext = SourcePipeline->bIsReimportContext;
+	bIsBasicLayout = SourcePipeline->bIsBasicLayout;
 }
 
 const FInterchangePipelinePropertyStates* UInterchangePipelineBase::GetPropertyStates(const FName PropertyPath) const
diff --git a/Engine/Source/Runtime/Interchange/Core/Public/InterchangePipelineBase.h b/Engine/Source/Runtime/Interchange/Core/Public/InterchangePipelineBase.h
--- a/Engine/Source/Runtime/Interchange/Core/Public/InterchangePipelineBase.h
+++ b/Engine/Source/Runtime/Interchange/Core/Public/InterchangePipelineBase.h
@@ -80,20 +80,45 @@ struct FInterchangePipelinePropertyStates
 	}
 	
 	/** Return true if the property is visible for the specified context. */
-	bool IsPropertyVisible(const bool bIsReimportContext) const
+	bool IsPropertyVisibleInBasicLayout() const
 	{
-		return bIsReimportContext ? ReimportStates.bVisible : ImportStates.bVisible;
+		return BasicLayoutStates.bVisible;
 	}
 
-	void SetPropertyVisible(const bool bIsReimportContext, const bool bVisibleValue)
+	/** Return true if the property is visible for the specified context. */
+	bool IsPropertyVisible(const bool bIsReimportContext, const bool bIsBasicLayout) const
+	{
+		bool bVisible = bIsReimportContext ? ReimportStates.bVisible : ImportStates.bVisible;
+		if (bVisible)
+		{
+			bVisible = bIsBasicLayout ? BasicLayoutStates.bVisible : true;
+		}
+		return bVisible;
+	}
+
+	void SetPropertyImportVisibility(const bool bVisibleValue)
+	{
+		ImportStates.bVisible = bVisibleValue;
+	}
+
+	void SetPropertyReimportVisibility(const bool bVisibleValue)
+	{
+		ReimportStates.bVisible = bVisibleValue;
+	}
+
+	void SetPropertyBasicLayoutVisibility(const bool bVisibleValue)
 	{
-		bIsReimportContext ? ReimportStates.bVisible = bVisibleValue : ImportStates.bVisible = bVisibleValue;
+		BasicLayoutStates.bVisible = bVisibleValue;
 	}
 
 	/** If true, the property is locked. */
 	UPROPERTY(EditAnywhere, BlueprintReadWrite, Category = "Property States")
 	bool bLocked = false;
 
+	/** The property states for the import context */
+	UPROPERTY(EditAnywhere, BlueprintReadWrite, Category = "Context Properties States")
+	FInterchangePipelinePropertyStatePerContext BasicLayoutStates;
+
 	/** The property states for the import context */
 	UPROPERTY(EditAnywhere, BlueprintReadWrite, Category = "Context Properties States")
 	FInterchangePipelinePropertyStatePerContext ImportStates;
@@ -257,6 +282,11 @@ public:
 	/** Transfer the source pipeline adjust settings to this pipeline. */
 	INTERCHANGECORE_API void TransferAdjustSettings(UInterchangePipelineBase* SourcePipeline);
 
+	INTERCHANGECORE_API void SetBasicLayoutMode(bool bBasicLayoutModeValue)
+	{
+		bIsBasicLayout = bBasicLayoutModeValue;
+	}
+
 	/**
 	 * This function is called before showing the import dialog it is not called doing a re-import.
 	 */
@@ -361,6 +391,7 @@ public:
 
 	bool CanEditPropertiesStates() { return bAllowPropertyStatesEdition; }
 	bool IsReimportContext() { return bIsReimportContext; }
+	bool IsBasicLayout() { return bIsBasicLayout; }
 
 #if WITH_EDITOR
 	/*
@@ -470,6 +501,14 @@ protected:
 	 */
 	bool bIsReimportContext = false;
 
+	/**
+	 * If true, this pipeline instance is use for basic layout.
+	 * If false, this pipeline instance is use for normal layout.
+	 *
+	 * Note: This layout must be set by the owner instancing this pipeline. This layout will be use to hide or not some properties.
+	 */
+	bool bIsBasicLayout = false;
+
 	UPROPERTY()
 	TObjectPtr<UInterchangeResultsContainer> Results;
 
