Interchange: Add translator settings system. Add force front x axis option to fbx translator

#rnx

diff --git a/Engine/Plugins/Interchange/Editor/Source/InterchangeEditor/InterchangeEditor.Build.cs b/Engine/Plugins/Interchange/Editor/Source/InterchangeEditor/InterchangeEditor.Build.cs
--- a/Engine/Plugins/Interchange/Editor/Source/InterchangeEditor/InterchangeEditor.Build.cs
+++ b/Engine/Plugins/Interchange/Editor/Source/InterchangeEditor/InterchangeEditor.Build.cs
@@ -23,9 +23,11 @@ namespace UnrealBuildTool.Rules
 				{
 					"AssetDefinition",
 					"ContentBrowser",
+					"InterchangeCommonParser",
 					"InterchangeCore",
 					"InterchangeEngine",
 					"InterchangeFactoryNodes",
+					"InterchangeFbxParser",
 					"InterchangeImport",
 					"InterchangeNodes",
 					"InterchangePipelines",
diff --git a/Engine/Plugins/Interchange/Editor/Source/InterchangeEditor/Private/InterchangeFbxAssetImportDataConverter.cpp b/Engine/Plugins/Interchange/Editor/Source/InterchangeEditor/Private/InterchangeFbxAssetImportDataConverter.cpp
--- a/Engine/Plugins/Interchange/Editor/Source/InterchangeEditor/Private/InterchangeFbxAssetImportDataConverter.cpp
+++ b/Engine/Plugins/Interchange/Editor/Source/InterchangeEditor/Private/InterchangeFbxAssetImportDataConverter.cpp
@@ -14,6 +14,7 @@
 #include "Factories/FbxSkeletalMeshImportData.h"
 #include "Factories/FbxStaticMeshImportData.h"
 #include "Factories/FbxTextureImportData.h"
+#include "Fbx/InterchangeFbxTranslator.h"
 #include "InterchangeAnimSequenceFactoryNode.h"
 #include "InterchangeAssetImportData.h"
 #include "InterchangeGenericAnimationPipeline.h"
@@ -35,11 +36,27 @@ namespace UE::Interchange::Private
 		DestinationData->SetSourceFiles(MoveTemp(SourceFiles));
 	}
 
-	void FillFbxAssetImportData(const UInterchangeGenericAssetsPipeline* GenericAssetPipeline, UFbxAssetImportData* AssetImportData)
+	void FillFbxAssetImportData(const UInterchangeFbxTranslatorSettings* InterchangeFbxTranslatorSettings, const UInterchangeGenericAssetsPipeline* GenericAssetPipeline, UFbxAssetImportData* AssetImportData)
 	{
-		AssetImportData->bConvertScene = true;
-		AssetImportData->bConvertSceneUnit = true;
-		AssetImportData->bForceFrontXAxis = false;
+		if (InterchangeFbxTranslatorSettings)
+		{
+			AssetImportData->bConvertScene = InterchangeFbxTranslatorSettings->bConvertScene;
+			AssetImportData->bConvertSceneUnit = InterchangeFbxTranslatorSettings->bConvertSceneUnit;
+			AssetImportData->bForceFrontXAxis = InterchangeFbxTranslatorSettings->bForceFrontXAxis;
+			
+		}
+		else if(UInterchangeFbxTranslatorSettings* InterchangeFbxTranslatorSettingsCDO = UInterchangeFbxTranslatorSettings::StaticClass()->GetDefaultObject<UInterchangeFbxTranslatorSettings>())
+		{
+			AssetImportData->bConvertScene = InterchangeFbxTranslatorSettingsCDO->bConvertScene;
+			AssetImportData->bConvertSceneUnit = InterchangeFbxTranslatorSettingsCDO->bConvertSceneUnit;
+			AssetImportData->bForceFrontXAxis = InterchangeFbxTranslatorSettingsCDO->bForceFrontXAxis;
+		}
+		else
+		{
+			AssetImportData->bConvertScene = true;
+			AssetImportData->bConvertSceneUnit = true;
+			AssetImportData->bForceFrontXAxis = false;
+		}
 		AssetImportData->bImportAsScene = false;
 		AssetImportData->ImportRotation = GenericAssetPipeline->ImportOffsetRotation;
 		AssetImportData->ImportTranslation = GenericAssetPipeline->ImportOffsetTranslation;
@@ -284,13 +301,15 @@ namespace UE::Interchange::Private
 		//Transfer the Source file information
 		TransferSourceFileInformation(InterchangeSourceData, DestinationStaticMeshImportData);
 
+		const UInterchangeFbxTranslatorSettings* InterchangeFbxTranslatorSettings = Cast<UInterchangeFbxTranslatorSettings>(InterchangeSourceData->GetTranslatorSettings());
+
 		//Now find the generic asset pipeline
 		for (UObject* Pipeline : InterchangeSourceData->GetPipelines())
 		{
 			if (const UInterchangeGenericAssetsPipeline* GenericAssetPipeline = Cast<UInterchangeGenericAssetsPipeline>(Pipeline))
 			{
 
-				FillFbxAssetImportData(GenericAssetPipeline, DestinationStaticMeshImportData);
+				FillFbxAssetImportData(InterchangeFbxTranslatorSettings, GenericAssetPipeline, DestinationStaticMeshImportData);
 				FillFbxMeshImportData(GenericAssetPipeline, DestinationStaticMeshImportData);
 
 				DestinationStaticMeshImportData->bAutoGenerateCollision = GenericAssetPipeline->MeshPipeline->bImportCollision;
@@ -347,12 +366,14 @@ namespace UE::Interchange::Private
 		//Transfer the Source file information
 		TransferSourceFileInformation(InterchangeSourceData, DestinationSkeletalMeshImportData);
 
+		const UInterchangeFbxTranslatorSettings* InterchangeFbxTranslatorSettings = Cast<UInterchangeFbxTranslatorSettings>(InterchangeSourceData->GetTranslatorSettings());
+
 		//Now find the generic asset pipeline
 		for (UObject* Pipeline : InterchangeSourceData->GetPipelines())
 		{
 			if (const UInterchangeGenericAssetsPipeline* GenericAssetPipeline = Cast<UInterchangeGenericAssetsPipeline>(Pipeline))
 			{
-				FillFbxAssetImportData(GenericAssetPipeline, DestinationSkeletalMeshImportData);
+				FillFbxAssetImportData(InterchangeFbxTranslatorSettings, GenericAssetPipeline, DestinationSkeletalMeshImportData);
 				FillFbxMeshImportData(GenericAssetPipeline, DestinationSkeletalMeshImportData);
 				DestinationSkeletalMeshImportData->bImportMeshesInBoneHierarchy = GenericAssetPipeline->CommonSkeletalMeshesAndAnimationsProperties->bImportMeshesInBoneHierarchy;
 				DestinationSkeletalMeshImportData->bImportMorphTargets = GenericAssetPipeline->MeshPipeline->bImportMorphTargets;
@@ -437,12 +458,14 @@ namespace UE::Interchange::Private
 		//Transfer the Source file information
 		TransferSourceFileInformation(InterchangeSourceData, DestinationAnimSequenceImportData);
 
+		const UInterchangeFbxTranslatorSettings* InterchangeFbxTranslatorSettings = Cast<UInterchangeFbxTranslatorSettings>(InterchangeSourceData->GetTranslatorSettings());
+
 		//Now find the generic asset pipeline
 		for (UObject* Pipeline : InterchangeSourceData->GetPipelines())
 		{
 			if (const UInterchangeGenericAssetsPipeline* GenericAssetPipeline = Cast<UInterchangeGenericAssetsPipeline>(Pipeline))
 			{
-				FillFbxAssetImportData(GenericAssetPipeline, DestinationAnimSequenceImportData);
+				FillFbxAssetImportData(InterchangeFbxTranslatorSettings, GenericAssetPipeline, DestinationAnimSequenceImportData);
 
 				switch (GenericAssetPipeline->AnimationPipeline->AnimationRange)
 				{
@@ -516,6 +539,13 @@ namespace UE::Interchange::Private
 		GenericAssetPipeline->ImportOffsetTranslation = FbxAssetImportData->ImportTranslation;
 		GenericAssetPipeline->ImportOffsetUniformScale = FbxAssetImportData->ImportUniformScale;
 
+		UInterchangeFbxTranslatorSettings* InterchangeFbxTranslatorSettings = NewObject<UInterchangeFbxTranslatorSettings>(DestinationData);
+		InterchangeFbxTranslatorSettings->ClearInternalFlags(EInternalObjectFlags::Async);
+		InterchangeFbxTranslatorSettings->bConvertScene = FbxAssetImportData->bConvertScene;
+		InterchangeFbxTranslatorSettings->bForceFrontXAxis = FbxAssetImportData->bForceFrontXAxis;
+		InterchangeFbxTranslatorSettings->bConvertSceneUnit = FbxAssetImportData->bConvertSceneUnit;
+		DestinationData->SetTranslatorSettings(InterchangeFbxTranslatorSettings);
+
 		if (const UFbxStaticMeshImportData* LegacyStaticMeshImportData = Cast<UFbxStaticMeshImportData>(FbxAssetImportData))
 		{
 			UInterchangeStaticMeshFactoryNode* MeshNode = NewObject<UInterchangeStaticMeshFactoryNode>(DestinationContainer);
@@ -577,6 +607,16 @@ namespace UE::Interchange::Private
 		Pipelines.Add(GenericAssetPipeline);
 		DestinationData->SetPipelines(Pipelines);
 
+		auto SetTranslatorSettings = [&DestinationData](const UFbxAssetImportData* FbxAssetImportData)
+			{
+				UInterchangeFbxTranslatorSettings* InterchangeFbxTranslatorSettings = NewObject<UInterchangeFbxTranslatorSettings>(DestinationData);
+				InterchangeFbxTranslatorSettings->ClearInternalFlags(EInternalObjectFlags::Async);
+				InterchangeFbxTranslatorSettings->bConvertScene = FbxAssetImportData->bConvertScene;
+				InterchangeFbxTranslatorSettings->bForceFrontXAxis = FbxAssetImportData->bForceFrontXAxis;
+				InterchangeFbxTranslatorSettings->bConvertSceneUnit = FbxAssetImportData->bConvertSceneUnit;
+				DestinationData->SetTranslatorSettings(InterchangeFbxTranslatorSettings);
+			};
+
 		//General Options
 		GenericAssetPipeline->bUseSourceNameForAsset = FbxImportUI->bOverrideFullName;
 		GenericAssetPipeline->ReimportStrategy = EReimportStrategyFlags::ApplyNoProperties;
@@ -632,6 +672,8 @@ namespace UE::Interchange::Private
 			GenericAssetPipeline->ImportOffsetTranslation = FbxImportUI->SkeletalMeshImportData->ImportTranslation;
 			GenericAssetPipeline->ImportOffsetUniformScale = FbxImportUI->SkeletalMeshImportData->ImportUniformScale;
 
+			SetTranslatorSettings(FbxImportUI->SkeletalMeshImportData);
+
 			FillInterchangeGenericAssetsPipelineFromFbxMeshImportData(GenericAssetPipeline, Cast<UFbxSkeletalMeshImportData>(FbxImportUI->SkeletalMeshImportData));
 		}
 		else if (FbxImportUI->MeshTypeToImport == EFBXImportType::FBXIT_StaticMesh)
@@ -646,6 +688,8 @@ namespace UE::Interchange::Private
 			GenericAssetPipeline->ImportOffsetTranslation = FbxImportUI->StaticMeshImportData->ImportTranslation;
 			GenericAssetPipeline->ImportOffsetUniformScale = FbxImportUI->StaticMeshImportData->ImportUniformScale;
 
+			SetTranslatorSettings(FbxImportUI->StaticMeshImportData);
+
 			FillInterchangeGenericAssetsPipelineFromFbxMeshImportData(GenericAssetPipeline, Cast<UFbxStaticMeshImportData>(FbxImportUI->StaticMeshImportData));
 		}
 		else if (FbxImportUI->MeshTypeToImport == EFBXImportType::FBXIT_Animation)
@@ -668,6 +712,10 @@ namespace UE::Interchange::Private
 			GenericAssetPipeline->ImportOffsetRotation = FbxImportUI->AnimSequenceImportData->ImportRotation;
 			GenericAssetPipeline->ImportOffsetTranslation = FbxImportUI->AnimSequenceImportData->ImportTranslation;
 			GenericAssetPipeline->ImportOffsetUniformScale = FbxImportUI->AnimSequenceImportData->ImportUniformScale;
+
+			SetTranslatorSettings(FbxImportUI->AnimSequenceImportData);
+
+			FillInterchangeGenericAssetsPipelineFromFbxAnimSequenceImportData(GenericAssetPipeline, Cast<UFbxAnimSequenceImportData>(FbxImportUI->AnimSequenceImportData));
 		}
 		else
 		{
@@ -676,6 +724,9 @@ namespace UE::Interchange::Private
 			GenericAssetPipeline->MeshPipeline->bImportStaticMeshes = true;
 			GenericAssetPipeline->MeshPipeline->bImportSkeletalMeshes = true;
 			GenericAssetPipeline->AnimationPipeline->bImportAnimations = true;
+
+			SetTranslatorSettings(FbxImportUI->StaticMeshImportData);
+
 			//Use the static mesh data
 			FillInterchangeGenericAssetsPipelineFromFbxMeshImportData(GenericAssetPipeline, Cast<UFbxStaticMeshImportData>(FbxImportUI->StaticMeshImportData));
 		}
diff --git a/Engine/Plugins/Interchange/Editor/Source/Pipelines/Private/InterchangePipelineConfigurationGeneric.cpp b/Engine/Plugins/Interchange/Editor/Source/Pipelines/Private/InterchangePipelineConfigurationGeneric.cpp
--- a/Engine/Plugins/Interchange/Editor/Source/Pipelines/Private/InterchangePipelineConfigurationGeneric.cpp
+++ b/Engine/Plugins/Interchange/Editor/Source/Pipelines/Private/InterchangePipelineConfigurationGeneric.cpp
@@ -12,6 +12,7 @@
 EInterchangePipelineConfigurationDialogResult UInterchangePipelineConfigurationGeneric::ShowPipelineConfigurationDialog(TArray<FInterchangeStackInfo>& PipelineStacks
 	, TArray<UInterchangePipelineBase*>& OutPipelines
 	, TWeakObjectPtr<UInterchangeSourceData> SourceData
+	, TWeakObjectPtr <UInterchangeTranslatorBase> Translator
 	, TWeakObjectPtr<UInterchangeBaseNodeContainer> BaseNodeContainer)
 {
 	//Create and show the graph inspector UI dialog
@@ -35,6 +36,7 @@ EInterchangePipelineConfigurationDialogResult UInterchangePipelineConfigurationG
 		.PipelineStacks(PipelineStacks)
 		.OutPipelines(&OutPipelines)
 		.BaseNodeContainer(BaseNodeContainer)
+		.Translator(Translator)
 	);
 
 	FSlateApplication::Get().AddModalWindow(Window, ParentWindow, false);
@@ -55,6 +57,7 @@ EInterchangePipelineConfigurationDialogResult UInterchangePipelineConfigurationG
 EInterchangePipelineConfigurationDialogResult UInterchangePipelineConfigurationGeneric::ShowScenePipelineConfigurationDialog(TArray<FInterchangeStackInfo>& PipelineStacks
 	, TArray<UInterchangePipelineBase*>& OutPipelines
 	, TWeakObjectPtr<UInterchangeSourceData> SourceData
+	, TWeakObjectPtr <UInterchangeTranslatorBase> Translator
 	, TWeakObjectPtr<UInterchangeBaseNodeContainer> BaseNodeContainer)
 {
 	//Create and show the graph inspector UI dialog
@@ -78,6 +81,7 @@ EInterchangePipelineConfigurationDialogResult UInterchangePipelineConfigurationG
 		.PipelineStacks(PipelineStacks)
 		.OutPipelines(&OutPipelines)
 		.BaseNodeContainer(BaseNodeContainer)
+		.Translator(Translator)
 	);
 
 	FSlateApplication::Get().AddModalWindow(Window, ParentWindow, false);
@@ -98,6 +102,7 @@ EInterchangePipelineConfigurationDialogResult UInterchangePipelineConfigurationG
 EInterchangePipelineConfigurationDialogResult UInterchangePipelineConfigurationGeneric::ShowReimportPipelineConfigurationDialog(TArray<FInterchangeStackInfo>& PipelineStacks
 	, TArray<UInterchangePipelineBase*>& OutPipelines
 	, TWeakObjectPtr<UInterchangeSourceData> SourceData
+	, TWeakObjectPtr <UInterchangeTranslatorBase> Translator
 	, TWeakObjectPtr<UInterchangeBaseNodeContainer> BaseNodeContainer
 	, TWeakObjectPtr <UObject> ReimportAsset)
 {
@@ -123,6 +128,7 @@ EInterchangePipelineConfigurationDialogResult UInterchangePipelineConfigurationG
 		.OutPipelines(&OutPipelines)
 		.BaseNodeContainer(BaseNodeContainer)
 		.ReimportObject(ReimportAsset)
+		.Translator(Translator)
 	);
 
 	FSlateApplication::Get().AddModalWindow(Window, ParentWindow, false);
diff --git a/Engine/Plugins/Interchange/Editor/Source/Pipelines/Private/SInterchangePipelineConfigurationDialog.cpp b/Engine/Plugins/Interchange/Editor/Source/Pipelines/Private/SInterchangePipelineConfigurationDialog.cpp
--- a/Engine/Plugins/Interchange/Editor/Source/Pipelines/Private/SInterchangePipelineConfigurationDialog.cpp
+++ b/Engine/Plugins/Interchange/Editor/Source/Pipelines/Private/SInterchangePipelineConfigurationDialog.cpp
@@ -2,18 +2,22 @@
 #include "SInterchangePipelineConfigurationDialog.h"
 
 #include "DetailsViewArgs.h"
+#include "Dialog/SCustomDialog.h"
 #include "Framework/MultiBox/MultiBoxBuilder.h"
 #include "Framework/Commands/UICommandList.h"
-#include "GameFramework/Actor.h"
 #include "Framework/Views/TableViewMetadata.h"
+#include "GameFramework/Actor.h"
 #include "IDetailsView.h"
 #include "IDocumentation.h"
 #include "InterchangeManager.h"
 #include "InterchangePipelineConfigurationBase.h"
 #include "InterchangeProjectSettings.h"
+#include "InterchangeTranslatorBase.h"
 #include "Interfaces/IMainFrameModule.h"
 #include "Misc/App.h"
 #include "Misc/ConfigCacheIni.h"
+#include "Misc/ConfigContext.h"
+#include "Misc/ScopedSlowTask.h"
 #include "Nodes/InterchangeBaseNodeContainer.h"
 #include "PropertyEditorModule.h"
 #include "SInterchangeGraphInspectorWindow.h"
@@ -37,28 +41,6 @@ static FAutoConsoleVariableRef CCvarInterchangeEnableFBXImport(
 const FName ReimportStackName = TEXT("ReimportPipeline");
 const FString ReimportPipelinePrefix = TEXT("reimport_");
 
- // Pipelines are renamed with the reimport prefix to avoid conflicts with the duplicates of the original pipelines that end up in the same package.
- // As this is the name displayed in the Dialog, conflicts won't matter.
-FString SInterchangePipelineConfigurationDialog::GetPipelineDisplayName(const UInterchangePipelineBase* Pipeline)
-{
-	static int32 RightChopIndex = ReimportPipelinePrefix.Len();
-
-	FString PipelineDisplayName = Pipeline->GetName();
-	if (PipelineDisplayName.StartsWith(ReimportPipelinePrefix))
-	{
-		PipelineDisplayName = PipelineDisplayName.RightChop(RightChopIndex);
-	}
-
-	FString StackName;
-	FString DisplayName;
-	if (PipelineDisplayName.Split("_", &StackName, &DisplayName))
-	{
-		return DisplayName;
-	}
-
-	return PipelineDisplayName;
-}
-
 void SInterchangePipelineItem::Construct(
 	const FArguments& InArgs,
 	const TSharedRef<STableViewBase>& OwnerTable,
@@ -209,6 +191,117 @@ SInterchangePipelineConfigurationDialog::~SInterchangePipelineConfigurationDialo
 	}
 }
 
+// Pipelines are renamed with the reimport prefix to avoid conflicts with the duplicates of the original pipelines that end up in the same package.
+ // As this is the name displayed in the Dialog, conflicts won't matter.
+FString SInterchangePipelineConfigurationDialog::GetPipelineDisplayName(const UInterchangePipelineBase* Pipeline)
+{
+	static int32 RightChopIndex = ReimportPipelinePrefix.Len();
+
+	FString PipelineDisplayName = Pipeline->GetName();
+	if (PipelineDisplayName.StartsWith(ReimportPipelinePrefix))
+	{
+		PipelineDisplayName = PipelineDisplayName.RightChop(RightChopIndex);
+	}
+
+	FString StackName;
+	FString DisplayName;
+	if (PipelineDisplayName.Split("_", &StackName, &DisplayName))
+	{
+		return DisplayName;
+	}
+
+	return PipelineDisplayName;
+}
+
+void SInterchangePipelineConfigurationDialog::SetEditPipeline(UInterchangePipelineBase* PipelineToEdit)
+{
+	TArray<UObject*> ObjectsToEdit;
+	ObjectsToEdit.Add(PipelineToEdit);
+	PipelineConfigurationDetailsView->SetObjects(ObjectsToEdit);
+}
+
+FReply SInterchangePipelineConfigurationDialog::OnEditTranslatorSettings()
+{
+	FPropertyEditorModule& PropertyEditorModule = FModuleManager::GetModuleChecked<FPropertyEditorModule>("PropertyEditor");
+	FDetailsViewArgs DetailsViewArgs;
+	DetailsViewArgs.bAllowSearch = true;
+	DetailsViewArgs.bShowPropertyMatrixButton = false;
+	DetailsViewArgs.bShowSectionSelector = false;
+	DetailsViewArgs.bAllowMultipleTopLevelObjects = true;
+	DetailsViewArgs.NameAreaSettings = FDetailsViewArgs::HideNameArea;
+	TSharedRef<IDetailsView> TranslatorSettingsDetailsView = PropertyEditorModule.CreateDetailView(DetailsViewArgs);
+	TranslatorSettingsDetailsView->OnFinishedChangingProperties().AddRaw(this, &SInterchangePipelineConfigurationDialog::OnFinishedChangingProperties);
+	TranslatorSettingsDetailsView->SetObject(TranslatorSettings);
+
+	TSharedRef<SCustomDialog> OptionsDialog =
+		SNew(SCustomDialog)
+		.Title(LOCTEXT("OptionsDialogTitle", "Translator Project Settings Editor"))
+		.WindowArguments(SWindow::FArguments()
+			.IsTopmostWindow(true)
+			.MinHeight(300.0f)
+			.MinWidth(500)
+			.SizingRule(ESizingRule::UserSized))
+		.HAlignContent(HAlign_Fill)
+		.VAlignContent(VAlign_Fill)
+		.HAlignIcon(HAlign_Left)
+		.VAlignIcon(VAlign_Top)
+		.UseScrollBox(true)
+		.Content()
+		[
+			SNew(SBorder)
+			.Padding(FMargin(10.0f, 3.0f))
+			.BorderImage(FAppStyle::GetBrush("ToolPanel.DarkGroupBorder"))
+			[
+				SNew(SVerticalBox)
+				+ SVerticalBox::Slot()
+				.HAlign(HAlign_Fill)
+				.VAlign(VAlign_Fill)
+				.FillHeight(1.0)
+				[
+					TranslatorSettingsDetailsView
+				]
+			]
+			
+		]
+		.Buttons(
+		{
+			SCustomDialog::FButton(LOCTEXT("DialogButtonOk", "Ok"))
+		});
+	OptionsDialog->ShowModal();
+
+	return FReply::Handled();
+}
+
+void SInterchangePipelineConfigurationDialog::OnFinishedChangingProperties(const FPropertyChangedEvent& PropertyChangedEvent)
+{
+	if (!Translator.IsValid() || !TranslatorSettings)
+	{
+		return;
+	}
+	if (UClass* TranslatorSettingsClass = TranslatorSettings->GetClass())
+	{
+		if (TranslatorSettingsClass->HasProperty(PropertyChangedEvent.Property))
+		{
+			//Save the config locally before the translation.
+			TranslatorSettings->SaveSettings();
+
+			//Need to Translate the source data
+ 			FScopedSlowTask Progress(2.f, NSLOCTEXT("SInterchangePipelineConfigurationDialog", "TranslatingSourceFile...", "Translating source file..."));
+ 			Progress.MakeDialog();
+ 			Progress.EnterProgressFrame(1.f);
+			//Reset the container
+			BaseNodeContainer->Reset();
+
+			Translator->Translate(*BaseNodeContainer.Get());
+
+			//Refresh the dialog
+			RefreshStack(false);
+
+			Progress.EnterProgressFrame(1.f);
+		}
+	}
+}
+
 TSharedRef<SBox> SInterchangePipelineConfigurationDialog::SpawnPipelineConfiguration()
 {
 	AvailableStacks.Reset();
@@ -365,10 +458,11 @@ TSharedRef<SBox> SInterchangePipelineConfigurationDialog::SpawnPipelineConfigura
 	DetailsViewArgs.bAllowSearch = true;
 	DetailsViewArgs.bShowPropertyMatrixButton = false;
 	DetailsViewArgs.bShowSectionSelector = true;
+	DetailsViewArgs.bAllowMultipleTopLevelObjects = true;
 	DetailsViewArgs.NameAreaSettings = FDetailsViewArgs::HideNameArea;
 	PipelineConfigurationDetailsView = PropertyEditorModule.CreateDetailView(DetailsViewArgs);
 	InspectorBox->SetContent(PipelineConfigurationDetailsView->AsShared());
-	PipelineConfigurationDetailsView->SetObject(nullptr);
+	SetEditPipeline(nullptr);
 	PipelineConfigurationDetailsView->GetIsPropertyVisibleDelegate().BindLambda([this](const FPropertyAndParent& PropertyAndParent)
 		{
 			return IsPropertyVisible(PropertyAndParent);
@@ -399,10 +493,17 @@ void SInterchangePipelineConfigurationDialog::Construct(const FArguments& InArgs
 	BaseNodeContainer = InArgs._BaseNodeContainer;
 	ReimportObject = InArgs._ReimportObject;
 	SourceData = InArgs._SourceData;
+	Translator = InArgs._Translator;
+	if (Translator.IsValid())
+	{
+		TranslatorSettings = Translator->GetSettings();
+	}
+
 	if (ReimportObject.IsValid())
 	{
 		ensure(bReimport);
 	}
+	
 
 	check(OutPipelines);
 
@@ -439,7 +540,26 @@ void SInterchangePipelineConfigurationDialog::Construct(const FArguments& InArgs
 			[
 				SNew(SHorizontalBox)
 				+SHorizontalBox::Slot()
+				.AutoWidth()
+				.VAlign(VAlign_Center)
+				.Padding(2.0f, 2.0f)
+				[
+					SNew(SButton)
+					.Visibility_Lambda([this]()
+						{
+							return !TranslatorSettings ? EVisibility::Collapsed : EVisibility::All;
+						})
+					.ToolTipText(LOCTEXT("SInterchangePipelineConfigurationDialog_TranslatorSettings_Tooltip", "Edit translator project settings."))
+					.OnClicked(this, &SInterchangePipelineConfigurationDialog::OnEditTranslatorSettings)
+					[
+						SNew(SImage)
+							.Image(FAppStyle::GetBrush("Icons.Settings"))
+							.ColorAndOpacity(FSlateColor::UseForeground())
+					]
+				]
+				+SHorizontalBox::Slot()
 				.FillWidth(1.0f)
+				.VAlign(VAlign_Center)
 				.Padding(0.0f, 2.0f)
 				[
 					SNew(STextBlock)
@@ -676,10 +796,11 @@ FReply SInterchangePipelineConfigurationDialog::OnResetToDefault()
 	{
 		return Result;
 	}
+	UInterchangePipelineBase* PipelineToEdit = nullptr;
 	//Multi selection is not allowed
-	ensure(SelectedPipelines.Num() <= 1);
 	for(TWeakObjectPtr<UObject> WeakObject : SelectedPipelines)
 	{
+		//We test the cast because we can have null or other type selected (i.e. translator settings class default object).
 		if (UInterchangePipelineBase* Pipeline = Cast<UInterchangePipelineBase>(WeakObject.Get()))
 		{
 			const UClass* PipelineClass = Pipeline->GetClass();
@@ -710,7 +831,7 @@ FReply SInterchangePipelineConfigurationDialog::OnResetToDefault()
 									}
 									//Switch the pipeline the element point on
 									PipelineListViewItems[PipelineIndex]->Pipeline = GeneratedPipeline;
-									PipelineConfigurationDetailsView->SetObject(GeneratedPipeline, true);
+									PipelineToEdit = GeneratedPipeline;
 									PipelinesListView->SetSelection(PipelineListViewItems[PipelineIndex], ESelectInfo::Direct);
 									PipelinesListView->RequestListRefresh();
 									break;
@@ -722,6 +843,7 @@ FReply SInterchangePipelineConfigurationDialog::OnResetToDefault()
 			}
 		}
 	}
+	SetEditPipeline(PipelineToEdit);
 	return Result;
 }
 
@@ -947,8 +1069,7 @@ void SInterchangePipelineConfigurationDialog::OnPipelineSelectionChanged(TShared
 	{
 		CurrentSelectedPipeline = InItem->Pipeline;
 	}
-	
-	PipelineConfigurationDetailsView->SetObject(CurrentSelectedPipeline.Get());
+	SetEditPipeline(CurrentSelectedPipeline.Get());
 	
 	if (CurrentSelectedPipeline)
 	{
@@ -1026,7 +1147,7 @@ FReply SInterchangePipelineConfigurationDialog::OnPreviewImport() const
 	//Create and show the graph inspector UI dialog
 	TSharedRef<SWindow> Window = SNew(SWindow)
 		.ClientSize(FVector2D(800.f, 650.f))
-		.Title(NSLOCTEXT("Interchange", "GraphInspectorTitle", "Interchange node graph inspector"));
+		.Title(NSLOCTEXT("SInterchangePipelineConfigurationDialog", "InterchangePreviewTitle", "Interchange Preview"));
 	TSharedPtr<SInterchangeGraphInspectorWindow> InterchangeGraphInspectorWindow;
 
 	Window->SetContent
diff --git a/Engine/Plugins/Interchange/Editor/Source/Pipelines/Private/SInterchangePipelineConfigurationDialog.h b/Engine/Plugins/Interchange/Editor/Source/Pipelines/Private/SInterchangePipelineConfigurationDialog.h
--- a/Engine/Plugins/Interchange/Editor/Source/Pipelines/Private/SInterchangePipelineConfigurationDialog.h
+++ b/Engine/Plugins/Interchange/Editor/Source/Pipelines/Private/SInterchangePipelineConfigurationDialog.h
@@ -15,6 +15,8 @@ struct FSlateBrush;
 class IDetailsView;
 class SCheckBox;
 class STextComboBox;
+class UInterchangeTranslatorBase;
+
 
 struct FInterchangePipelineItemType
 {
@@ -68,6 +70,7 @@ public:
 		SLATE_ARGUMENT(TArray<UInterchangePipelineBase*>*, OutPipelines)
 		SLATE_ARGUMENT(TWeakObjectPtr<UInterchangeBaseNodeContainer>, BaseNodeContainer)
 		SLATE_ARGUMENT(TWeakObjectPtr<UObject>, ReimportObject)
+		SLATE_ARGUMENT(TWeakObjectPtr<UInterchangeTranslatorBase>, Translator)
 	SLATE_END_ARGS()
 
 public:
@@ -112,11 +115,17 @@ private:
 	/** Internal utility function to properly display pipeline's name */
 	static FString GetPipelineDisplayName(const UInterchangePipelineBase* Pipeline);
 
+	void SetEditPipeline(UInterchangePipelineBase* PipelineToEdit);
+	FReply OnEditTranslatorSettings();
+	void OnFinishedChangingProperties(const FPropertyChangedEvent& PropertyChangedEvent);
+
 private:
 	TWeakPtr< SWindow > OwnerWindow;
 	TWeakObjectPtr<UInterchangeSourceData> SourceData;
 	TWeakObjectPtr<UInterchangeBaseNodeContainer> BaseNodeContainer;
 	TWeakObjectPtr<UObject> ReimportObject;
+	TWeakObjectPtr<UInterchangeTranslatorBase> Translator;
+	TObjectPtr<UInterchangeTranslatorSettings> TranslatorSettings = nullptr;
 	TArray<FInterchangeStackInfo> PipelineStacks;
 	TArray<UInterchangePipelineBase*>* OutPipelines;
 
diff --git a/Engine/Plugins/Interchange/Editor/Source/Pipelines/Public/InterchangePipelineConfigurationGeneric.h b/Engine/Plugins/Interchange/Editor/Source/Pipelines/Public/InterchangePipelineConfigurationGeneric.h
--- a/Engine/Plugins/Interchange/Editor/Source/Pipelines/Public/InterchangePipelineConfigurationGeneric.h
+++ b/Engine/Plugins/Interchange/Editor/Source/Pipelines/Public/InterchangePipelineConfigurationGeneric.h
@@ -21,14 +21,17 @@ protected:
 	virtual EInterchangePipelineConfigurationDialogResult ShowPipelineConfigurationDialog(TArray<FInterchangeStackInfo>& PipelineStacks
 		, TArray<UInterchangePipelineBase*>& OutPipelines
 		, TWeakObjectPtr<UInterchangeSourceData> SourceData
+		, TWeakObjectPtr <UInterchangeTranslatorBase> Translator
 		, TWeakObjectPtr<UInterchangeBaseNodeContainer> BaseNodeContainer) override;
 	virtual EInterchangePipelineConfigurationDialogResult ShowScenePipelineConfigurationDialog(TArray<FInterchangeStackInfo>& PipelineStacks
 		, TArray<UInterchangePipelineBase*>& OutPipelines
 		, TWeakObjectPtr<UInterchangeSourceData> SourceData
+		, TWeakObjectPtr <UInterchangeTranslatorBase> Translator
 		, TWeakObjectPtr<UInterchangeBaseNodeContainer> BaseNodeContainer) override;
 	virtual EInterchangePipelineConfigurationDialogResult ShowReimportPipelineConfigurationDialog(TArray<FInterchangeStackInfo>& PipelineStacks
 		, TArray<UInterchangePipelineBase*>& OutPipelines
 		, TWeakObjectPtr<UInterchangeSourceData> SourceData
+		, TWeakObjectPtr <UInterchangeTranslatorBase> Translator
 		, TWeakObjectPtr<UInterchangeBaseNodeContainer> BaseNodeContainer
 		, TWeakObjectPtr <UObject> ReimportAsset) override;
 };
diff --git a/Engine/Plugins/Interchange/Runtime/Config/BaseInterchange.ini b/Engine/Plugins/Interchange/Runtime/Config/BaseInterchange.ini
--- a/Engine/Plugins/Interchange/Runtime/Config/BaseInterchange.ini
+++ b/Engine/Plugins/Interchange/Runtime/Config/BaseInterchange.ini
@@ -1,3 +1,6 @@
+[SectionsToSave]
+bCanSaveAllSections=True
+
 [CoreRedirects]
 +EnumRedirects=(OldName="/Script/InterchangePipelines.EInterchangeMaterialImportOption",ValueChanges=(("DoNotImport","ImportAsMaterials"))
 
@@ -55,3 +58,8 @@ PredefinedSurfaceShaders=((StandardSurface, "/Interchange/Functions/MX_StandardS
 PredefinedBSDF=((OrenNayarDiffuse, "/Interchange/Functions/MX_OrenNayarBSDF.MX_OrenNayarBSDF"),(BurleyDiffuse,"/Interchange/Functions/MX_BurleyDiffuseBSDF.MX_BurleyDiffuseBSDF"),(Translucent,"/Interchange/Functions/MX_TranslucentBSDF.MX_TranslucentBSDF"),(Dielectric,"/Interchange/Functions/MX_DielectricBSDF.MX_DielectricBSDF"),(Conductor,"/Interchange/Functions/MX_ConductorBSDF.MX_ConductorBSDF"),(GeneralizedSchlick,"/Interchange/Functions/MX_GeneralizedSchlickBSDF.MX_GeneralizedSchlickBSDF"),(Subsurface,"/Interchange/Functions/MX_SubsurfaceBSDF.MX_SubsurfaceBSDF"),(Sheen,"/Interchange/Functions/MX_SheenBSDF.MX_SheenBSDF"),(ThinFilm,"/Interchange/Functions/MX_ThinFilmBSDF.MX_ThinFilmBSDF"))
 PredefinedEDF=((Uniform, "/Interchange/Functions/MX_UniformEDF.MX_UniformEDF"),(Conical,"/Interchange/Functions/MX_ConicalEDF.MX_ConicalEDF"),(Measured,"/Interchange/Functions/MX_MeasuredEDF.MX_MeasuredEDF"))
 PredefinedVDF=((Absorption, "/Interchange/Functions/MX_AbsorptionVDF.MX_AbsorptionVDF"),(Anisotropic,"/Interchange/Functions/MX_AnisotropicVDF.MX_AnisotropicVDF"))
+
+[/Script/InterchangeImport.InterchangeFbxTranslatorProjectSettings]
+bConvertScene=True
+bForceFrontXAxis=False
+bConvertSceneUnit=True
\ No newline at end of file
diff --git a/Engine/Plugins/Interchange/Runtime/Source/Dispatcher/Private/InterchangeDispatcherTask.cpp b/Engine/Plugins/Interchange/Runtime/Source/Dispatcher/Private/InterchangeDispatcherTask.cpp
--- a/Engine/Plugins/Interchange/Runtime/Source/Dispatcher/Private/InterchangeDispatcherTask.cpp
+++ b/Engine/Plugins/Interchange/Runtime/Source/Dispatcher/Private/InterchangeDispatcherTask.cpp
@@ -24,6 +24,9 @@ namespace UE
 			CmdObject->SetStringField(GetCommandIDJsonKey(), GetAction());
 			CmdObject->SetStringField(GetTranslatorIDJsonKey(), GetTranslatorID());
 			ActionDataObject->SetStringField(GetSourceFilenameJsonKey(), GetSourceFilename());
+			ActionDataObject->SetBoolField(GetDoesConvertSceneJsonKey(), GetDoesConvertScene());
+			ActionDataObject->SetBoolField(GetDoesForceFrontXAxisJsonKey(), GetDoesForceFrontXAxis());
+			ActionDataObject->SetBoolField(GetDoesConvertSceneUnitJsonKey(), GetDoesConvertSceneUnit());
 			CmdObject->SetObjectField(GetCommandDataJsonKey(), ActionDataObject);
 
 			FString LoadSourceCmd;
@@ -75,6 +78,18 @@ namespace UE
 			{
 				return false;
 			}
+			if (!((*ActionDataObject)->TryGetBoolField(GetDoesConvertSceneJsonKey(), bConvertScene)))
+			{
+				return false;
+			}
+			if (!((*ActionDataObject)->TryGetBoolField(GetDoesForceFrontXAxisJsonKey(), bForceFrontXAxis)))
+			{
+				return false;
+			}
+			if (!((*ActionDataObject)->TryGetBoolField(GetDoesConvertSceneUnitJsonKey(), bConvertSceneUnit)))
+			{
+				return false;
+			}
 
 			//Since we filled the data from the json file, set the data has been initialize.
 			bIsDataInitialize = true;
diff --git a/Engine/Plugins/Interchange/Runtime/Source/Dispatcher/Public/InterchangeDispatcherTask.h b/Engine/Plugins/Interchange/Runtime/Source/Dispatcher/Public/InterchangeDispatcherTask.h
--- a/Engine/Plugins/Interchange/Runtime/Source/Dispatcher/Public/InterchangeDispatcherTask.h
+++ b/Engine/Plugins/Interchange/Runtime/Source/Dispatcher/Public/InterchangeDispatcherTask.h
@@ -91,9 +91,16 @@ namespace UE
 				bIsDataInitialize = false;
 			}
 
-			FJsonLoadSourceCmd(const FString& InTranslatorID, const FString& InSourceFilename)
+			FJsonLoadSourceCmd(const FString& InTranslatorID
+				, const FString& InSourceFilename
+				, const bool InbConvertScene
+				, const bool InbForceFrontXAxis
+				, const bool InbConvertSceneUnit)
 				: TranslatorID(InTranslatorID)
 				, SourceFilename(InSourceFilename)
+				, bConvertScene(InbConvertScene)
+				, bForceFrontXAxis(InbForceFrontXAxis)
+				, bConvertSceneUnit(InbConvertSceneUnit)
 			{
 				bIsDataInitialize = true;
 			}
@@ -127,6 +134,45 @@ namespace UE
 				return Key;
 			}
 
+			bool GetDoesConvertScene() const
+			{
+				//Code should not do query data if the data was not set before
+				ensure(bIsDataInitialize);
+				return bConvertScene;
+			}
+
+			static FString GetDoesConvertSceneJsonKey()
+			{
+				static const FString Key = TEXT("ConvertScene");
+				return Key;
+			}
+
+			bool GetDoesForceFrontXAxis() const
+			{
+				//Code should not do query data if the data was not set before
+				ensure(bIsDataInitialize);
+				return bForceFrontXAxis;
+			}
+
+			static FString GetDoesForceFrontXAxisJsonKey()
+			{
+				static const FString Key = TEXT("ForceFrontXAxis");
+				return Key;
+			}
+
+			bool GetDoesConvertSceneUnit() const
+			{
+				//Code should not do query data if the data was not set before
+				ensure(bIsDataInitialize);
+				return bConvertSceneUnit;
+			}
+
+			static FString GetDoesConvertSceneUnitJsonKey()
+			{
+				static const FString Key = TEXT("ConvertSceneUnit");
+				return Key;
+			}
+
 			/**
 			 * Use this class helper to create the cmd result json string and to read it
 			 */
@@ -156,6 +202,9 @@ namespace UE
 		private:
 			FString TranslatorID = FString();
 			FString SourceFilename = FString();
+			bool bConvertScene = true;
+			bool bForceFrontXAxis = false;
+			bool bConvertSceneUnit = true;
 		};
 
 		class INTERCHANGEDISPATCHER_API FJsonFetchPayloadCmd : public IJsonCmdBase
diff --git a/Engine/Plugins/Interchange/Runtime/Source/Import/Private/Animation/InterchangeAnimSequenceFactory.cpp b/Engine/Plugins/Interchange/Runtime/Source/Import/Private/Animation/InterchangeAnimSequenceFactory.cpp
--- a/Engine/Plugins/Interchange/Runtime/Source/Import/Private/Animation/InterchangeAnimSequenceFactory.cpp
+++ b/Engine/Plugins/Interchange/Runtime/Source/Import/Private/Animation/InterchangeAnimSequenceFactory.cpp
@@ -1358,7 +1358,7 @@ void UInterchangeAnimSequenceFactory::SetupObject_GameThread(const FSetupObjectP
 		check(AnimSequence == CastChecked<UAnimSequence>(Arguments.ImportedObject));
 
 		UAssetImportData* ImportDataPtr = AnimSequence->AssetImportData;
-		UE::Interchange::FFactoryCommon::FUpdateImportAssetDataParameters UpdateImportAssetDataParameters(AnimSequence, ImportDataPtr, Arguments.SourceData, Arguments.NodeUniqueID, Arguments.NodeContainer, Arguments.OriginalPipelines);
+		UE::Interchange::FFactoryCommon::FUpdateImportAssetDataParameters UpdateImportAssetDataParameters(AnimSequence, ImportDataPtr, Arguments.SourceData, Arguments.NodeUniqueID, Arguments.NodeContainer, Arguments.OriginalPipelines, Arguments.Translator);
 		ImportDataPtr = UE::Interchange::FFactoryCommon::UpdateImportAssetData(UpdateImportAssetDataParameters);
 		AnimSequence->AssetImportData = ImportDataPtr;
 	}
diff --git a/Engine/Plugins/Interchange/Runtime/Source/Import/Private/Fbx/InterchangeFbxTranslator.cpp b/Engine/Plugins/Interchange/Runtime/Source/Import/Private/Fbx/InterchangeFbxTranslator.cpp
--- a/Engine/Plugins/Interchange/Runtime/Source/Import/Private/Fbx/InterchangeFbxTranslator.cpp
+++ b/Engine/Plugins/Interchange/Runtime/Source/Import/Private/Fbx/InterchangeFbxTranslator.cpp
@@ -180,6 +180,12 @@ bool UInterchangeFbxTranslator::Translate(UInterchangeBaseNodeContainer& BaseNod
 		return false;
 	}
 
+	GetSettings();
+	ensure(CacheFbxTranslatorSettings);
+	const bool bConvertScene = CacheFbxTranslatorSettings ? CacheFbxTranslatorSettings->bConvertScene : true;
+	const bool bForceFrontXAxis = CacheFbxTranslatorSettings ? CacheFbxTranslatorSettings->bForceFrontXAxis : false;
+	const bool bConvertSceneUnit = CacheFbxTranslatorSettings ? CacheFbxTranslatorSettings->bConvertSceneUnit : true;
+
 	if (bUseWorkerImport)
 	{
 		if (!Dispatcher.IsValid())
@@ -187,7 +193,7 @@ bool UInterchangeFbxTranslator::Translate(UInterchangeBaseNodeContainer& BaseNod
 			return false;
 		}
 		//Create a json command to read the fbx file
-		FString JsonCommand = CreateLoadFbxFileCommand(Filename);
+		FString JsonCommand = CreateLoadFbxFileCommand(Filename, bConvertScene, bForceFrontXAxis, bConvertSceneUnit);
 		int32 TaskIndex = Dispatcher->AddTask(JsonCommand);
 
 		//Blocking call until all tasks are executed
@@ -225,7 +231,9 @@ bool UInterchangeFbxTranslator::Translate(UInterchangeBaseNodeContainer& BaseNod
 	else
 	{
 #if WITH_EDITOR
+		FbxParser.Reset();
 		FbxParser.SetResultContainer(Results);
+		FbxParser.SetConvertSettings(bConvertScene, bForceFrontXAxis, bConvertSceneUnit);
 		FbxParser.LoadFbxFile(Filename, BaseNodeContainer);
 #endif
 	}
@@ -251,6 +259,12 @@ void UInterchangeFbxTranslator::ReleaseSource()
 		IFileManager::Get().DeleteDirectory(*ResultFolder, RequireExists, Tree);
 	}
 	ResultFolder.Empty();
+
+	if (CacheFbxTranslatorSettings)
+	{
+		CacheFbxTranslatorSettings->ClearFlags(RF_Standalone);
+		CacheFbxTranslatorSettings = nullptr;
+	}
 }
 
 void UInterchangeFbxTranslator::ImportFinish()
@@ -261,6 +275,33 @@ void UInterchangeFbxTranslator::ImportFinish()
 	}
 }
 
+UInterchangeTranslatorSettings* UInterchangeFbxTranslator::GetSettings() const
+{
+	if (!CacheFbxTranslatorSettings)
+	{
+		CacheFbxTranslatorSettings = DuplicateObject<UInterchangeFbxTranslatorSettings>(UInterchangeFbxTranslatorSettings::StaticClass()->GetDefaultObject<UInterchangeFbxTranslatorSettings>(), GetTransientPackage());
+		CacheFbxTranslatorSettings->LoadSettings();
+		CacheFbxTranslatorSettings->SetFlags(RF_Standalone);
+		CacheFbxTranslatorSettings->ClearInternalFlags(EInternalObjectFlags::Async);
+	}
+	return CacheFbxTranslatorSettings;
+}
+
+void UInterchangeFbxTranslator::SetSettings(const UInterchangeTranslatorSettings* InterchangeTranslatorSettings)
+{
+	if (CacheFbxTranslatorSettings)
+	{
+		CacheFbxTranslatorSettings->ClearFlags(RF_Standalone);
+		CacheFbxTranslatorSettings->ClearInternalFlags(EInternalObjectFlags::Async);
+		CacheFbxTranslatorSettings = nullptr;
+	}
+	if (InterchangeTranslatorSettings)
+	{
+		CacheFbxTranslatorSettings = DuplicateObject<UInterchangeFbxTranslatorSettings>(Cast<UInterchangeFbxTranslatorSettings>(InterchangeTranslatorSettings), GetTransientPackage());
+		CacheFbxTranslatorSettings->ClearInternalFlags(EInternalObjectFlags::Async);
+		CacheFbxTranslatorSettings->SetFlags(RF_Standalone);
+	}
+}
 
 TOptional<UE::Interchange::FImportImage> UInterchangeFbxTranslator::GetTexturePayloadData(const FString& PayLoadKey, TOptional<FString>& AlternateTexturePath) const
 {
@@ -556,9 +597,9 @@ TFuture<TOptional<UE::Interchange::FAnimationPayloadData>> UInterchangeFbxTransl
 	return Promise->GetFuture();
 }
 
-FString UInterchangeFbxTranslator::CreateLoadFbxFileCommand(const FString& FbxFilePath) const
+FString UInterchangeFbxTranslator::CreateLoadFbxFileCommand(const FString& FbxFilePath, const bool bConvertScene, const bool bForceFrontXAxis, const bool bConvertSceneUnit) const
 {
-	UE::Interchange::FJsonLoadSourceCmd LoadSourceCommand(TEXT("FBX"), FbxFilePath);
+	UE::Interchange::FJsonLoadSourceCmd LoadSourceCommand(TEXT("FBX"), FbxFilePath, bConvertScene, bForceFrontXAxis, bConvertSceneUnit);
 	return LoadSourceCommand.ToJson();
 }
 
diff --git a/Engine/Plugins/Interchange/Runtime/Source/Import/Private/Material/InterchangeMaterialFactory.cpp b/Engine/Plugins/Interchange/Runtime/Source/Import/Private/Material/InterchangeMaterialFactory.cpp
--- a/Engine/Plugins/Interchange/Runtime/Source/Import/Private/Material/InterchangeMaterialFactory.cpp
+++ b/Engine/Plugins/Interchange/Runtime/Source/Import/Private/Material/InterchangeMaterialFactory.cpp
@@ -1138,7 +1138,8 @@ void UInterchangeMaterialFactory::SetupObject_GameThread(const FSetupObjectParam
 																										  , Arguments.SourceData
 																										  , Arguments.NodeUniqueID
 																										  , Arguments.NodeContainer
-																										  , Arguments.OriginalPipelines);
+																										  , Arguments.OriginalPipelines
+																										  , Arguments.Translator);
 
 		ImportedMaterialInterface->AssetImportData = UE::Interchange::FFactoryCommon::UpdateImportAssetData(UpdateImportAssetDataParameters);
 #endif // WITH_EDITORONLY_DATA
diff --git a/Engine/Plugins/Interchange/Runtime/Source/Import/Private/Mesh/InterchangeSkeletalMeshFactory.cpp b/Engine/Plugins/Interchange/Runtime/Source/Import/Private/Mesh/InterchangeSkeletalMeshFactory.cpp
--- a/Engine/Plugins/Interchange/Runtime/Source/Import/Private/Mesh/InterchangeSkeletalMeshFactory.cpp
+++ b/Engine/Plugins/Interchange/Runtime/Source/Import/Private/Mesh/InterchangeSkeletalMeshFactory.cpp
@@ -1900,7 +1900,8 @@ void UInterchangeSkeletalMeshFactory::SetupObject_GameThread(const FSetupObjectP
 																										  , Arguments.SourceData
 																										  , Arguments.NodeUniqueID
 																										  , Arguments.NodeContainer
-																										  , Arguments.OriginalPipelines);
+																										  , Arguments.OriginalPipelines
+																										  , Arguments.Translator);
 
 		ImportDataPtr = UE::Interchange::FFactoryCommon::UpdateImportAssetData(UpdateImportAssetDataParameters, [&Arguments, SkeletalMesh](UInterchangeAssetImportData* AssetImportData)
 			{
diff --git a/Engine/Plugins/Interchange/Runtime/Source/Import/Private/Mesh/InterchangeStaticMeshFactory.cpp b/Engine/Plugins/Interchange/Runtime/Source/Import/Private/Mesh/InterchangeStaticMeshFactory.cpp
--- a/Engine/Plugins/Interchange/Runtime/Source/Import/Private/Mesh/InterchangeStaticMeshFactory.cpp
+++ b/Engine/Plugins/Interchange/Runtime/Source/Import/Private/Mesh/InterchangeStaticMeshFactory.cpp
@@ -833,7 +833,7 @@ void UInterchangeStaticMeshFactory::SetupObject_GameThread(const FSetupObjectPar
 		UStaticMesh* StaticMesh = CastChecked<UStaticMesh>(Arguments.ImportedObject);
 
 		UAssetImportData* ImportDataPtr = StaticMesh->GetAssetImportData();
-		UE::Interchange::FFactoryCommon::FUpdateImportAssetDataParameters UpdateImportAssetDataParameters(StaticMesh, ImportDataPtr, Arguments.SourceData, Arguments.NodeUniqueID, Arguments.NodeContainer, Arguments.OriginalPipelines);
+		UE::Interchange::FFactoryCommon::FUpdateImportAssetDataParameters UpdateImportAssetDataParameters(StaticMesh, ImportDataPtr, Arguments.SourceData, Arguments.NodeUniqueID, Arguments.NodeContainer, Arguments.OriginalPipelines, Arguments.Translator);
 		ImportDataPtr = UE::Interchange::FFactoryCommon::UpdateImportAssetData(UpdateImportAssetDataParameters);
 		StaticMesh->SetAssetImportData(ImportDataPtr);
 	}
diff --git a/Engine/Plugins/Interchange/Runtime/Source/Import/Private/Scene/InterchangeSceneImportAssetFactory.cpp b/Engine/Plugins/Interchange/Runtime/Source/Import/Private/Scene/InterchangeSceneImportAssetFactory.cpp
--- a/Engine/Plugins/Interchange/Runtime/Source/Import/Private/Scene/InterchangeSceneImportAssetFactory.cpp
+++ b/Engine/Plugins/Interchange/Runtime/Source/Import/Private/Scene/InterchangeSceneImportAssetFactory.cpp
@@ -129,7 +129,7 @@ void UInterchangeSceneImportAssetFactory::SetupObject_GameThread(const FSetupObj
 #if WITH_EDITORONLY_DATA
 		using namespace UE::Interchange;
 
-		FFactoryCommon::FUpdateImportAssetDataParameters Parameters(SceneImportAsset, SceneImportAsset->AssetImportData, Arguments.SourceData, Arguments.NodeUniqueID, Arguments.NodeContainer, Arguments.OriginalPipelines);
+		FFactoryCommon::FUpdateImportAssetDataParameters Parameters(SceneImportAsset, SceneImportAsset->AssetImportData, Arguments.SourceData, Arguments.NodeUniqueID, Arguments.NodeContainer, Arguments.OriginalPipelines, Arguments.Translator);
 		SceneImportAsset->AssetImportData = Cast<UInterchangeAssetImportData>(FFactoryCommon::UpdateImportAssetData(Parameters));
 #endif
 	}
diff --git a/Engine/Plugins/Interchange/Runtime/Source/Import/Private/Texture/InterchangeTextureFactory.cpp b/Engine/Plugins/Interchange/Runtime/Source/Import/Private/Texture/InterchangeTextureFactory.cpp
--- a/Engine/Plugins/Interchange/Runtime/Source/Import/Private/Texture/InterchangeTextureFactory.cpp
+++ b/Engine/Plugins/Interchange/Runtime/Source/Import/Private/Texture/InterchangeTextureFactory.cpp
@@ -1880,7 +1880,8 @@ void UInterchangeTextureFactory::SetupObject_GameThread(const FSetupObjectParams
 																						 , Arguments.SourceData
 																						 , Arguments.NodeUniqueID
 																						 , Arguments.NodeContainer
-																						 , Arguments.OriginalPipelines);
+																						 , Arguments.OriginalPipelines
+																						 , Arguments.Translator);
 		SetImportAssetDataParameters.SourceFiles = MoveTemp(SourceFiles);
 
 		Texture->AssetImportData = UE::Interchange::FFactoryCommon::SetImportAssetData(SetImportAssetDataParameters);
diff --git a/Engine/Plugins/Interchange/Runtime/Source/Import/Public/Fbx/InterchangeFbxTranslator.h b/Engine/Plugins/Interchange/Runtime/Source/Import/Public/Fbx/InterchangeFbxTranslator.h
--- a/Engine/Plugins/Interchange/Runtime/Source/Import/Public/Fbx/InterchangeFbxTranslator.h
+++ b/Engine/Plugins/Interchange/Runtime/Source/Import/Public/Fbx/InterchangeFbxTranslator.h
@@ -1,6 +1,7 @@
 // Copyright Epic Games, Inc. All Rights Reserved.
 #pragma once
 
+#include "Engine/DeveloperSettings.h"
 #include "Animation/InterchangeAnimationPayloadInterface.h"
 #include "CoreMinimal.h"
 #include "InterchangeTranslatorBase.h"
@@ -20,6 +21,26 @@
 
 /* Fbx translator class support import of texture, material, static mesh, skeletal mesh, */
 
+UCLASS(BlueprintType, editinlinenew, MinimalAPI)
+class UInterchangeFbxTranslatorSettings : public UInterchangeTranslatorSettings
+{
+	GENERATED_BODY()
+
+public:
+
+	/** Whether to convert FBX scene axis system to Unreal axis system. */
+	UPROPERTY(EditAnywhere, Category = "Fbx Translator")
+	bool bConvertScene = true;
+
+	/** Whether to force the front axis to be align with X instead of -Y default. */
+	UPROPERTY(EditAnywhere, Category = "Fbx Translator", meta = (EditCondition = "bConvertScene"))
+	bool bForceFrontXAxis = false;
+
+	/** Whether to convert the scene from FBX unit to UE unit (centimeter). */
+	UPROPERTY(EditAnywhere, Category = "Fbx Translator")
+	bool bConvertSceneUnit = true;
+};
+
 UCLASS(BlueprintType)
 class INTERCHANGEIMPORT_API UInterchangeFbxTranslator : public UInterchangeTranslatorBase
 , public IInterchangeTexturePayloadInterface
@@ -40,6 +61,8 @@ public:
 	virtual bool Translate(UInterchangeBaseNodeContainer& BaseNodeContainer) const override;
 	virtual void ReleaseSource() override;
 	virtual void ImportFinish() override;
+	virtual UInterchangeTranslatorSettings* GetSettings() const override;
+	virtual void SetSettings(const UInterchangeTranslatorSettings* InterchangeTranslatorSettings) override;
 	/** End UInterchangeTranslatorBase API*/
 
 
@@ -79,7 +102,7 @@ public:
 	virtual TFuture<TOptional<UE::Interchange::FAnimationPayloadData>> GetAnimationPayloadData(const FInterchangeAnimationPayLoadKey& PayLoadKey, const double BakeFrequency = 0, const double RangeStartSecond = 0, const double RangeStopSecond = 0) const override;
 	/* IInterchangeAnimationPayloadInterface End */
 private:
-	FString CreateLoadFbxFileCommand(const FString& FbxFilePath) const;
+	FString CreateLoadFbxFileCommand(const FString& FbxFilePath, const bool bConvertScene, const bool bForceFrontXAxis, const bool bConvertSceneUnit) const;
 
 	FString CreateFetchMeshPayloadFbxCommand(const FString& FbxPayloadKey, const FTransform& MeshGlobalTransform) const;
 
@@ -93,6 +116,8 @@ private:
 	//never translate a source.
 	mutable TUniquePtr<UE::Interchange::FInterchangeDispatcher> Dispatcher;
 
+	mutable TObjectPtr<UInterchangeFbxTranslatorSettings> CacheFbxTranslatorSettings = nullptr;
+
 	//If true this translator will use the dispatcher (InterchangeWorker program) to translate and return payloads.
 	//If false, this translator will not use the dispatcher
 	bool bUseWorkerImport = false;
@@ -102,4 +127,3 @@ private:
 	FString ResultFolder;
 };
 
-
diff --git a/Engine/Plugins/Interchange/Runtime/Source/Parsers/Fbx/Private/FbxAPI.cpp b/Engine/Plugins/Interchange/Runtime/Source/Parsers/Fbx/Private/FbxAPI.cpp
--- a/Engine/Plugins/Interchange/Runtime/Source/Parsers/Fbx/Private/FbxAPI.cpp
+++ b/Engine/Plugins/Interchange/Runtime/Source/Parsers/Fbx/Private/FbxAPI.cpp
@@ -35,8 +35,13 @@ namespace UE
 		{
 			FFbxParser::~FFbxParser()
 			{
-				PayloadContexts.Empty();
 				FbxHelper = nullptr;
+				Reset();
+			}
+			
+			void FFbxParser::Reset()
+			{
+				PayloadContexts.Reset();
 
 				DESTROY_FBX_OBJECT(SDKImporter);
 				DESTROY_FBX_OBJECT(SDKScene);
@@ -47,8 +52,12 @@ namespace UE
 				}
 				DESTROY_FBX_OBJECT(SDKIoSettings);
 				DESTROY_FBX_OBJECT(SDKManager);
+				if (FbxHelper.IsValid())
+				{
+					FbxHelper->Reset();
+				}
 			}
-			
+
 			const TSharedPtr<FFbxHelper> FFbxParser::GetFbxHelper()
 			{
 				if (!FbxHelper.IsValid())
@@ -121,7 +130,7 @@ namespace UE
 				bool bStatus = SDKImporter->Import(SDKScene);
 
 				//We always convert scene to UE axis and units
-				FFbxConvert::ConvertScene(SDKScene);
+				FFbxConvert::ConvertScene(SDKScene, bConvertScene, bForceFrontXAxis, bConvertSceneUnit);
 
 				FrameRate = FbxTime::GetFrameRate(SDKScene->GetGlobalSettings().GetTimeMode());
 
diff --git a/Engine/Plugins/Interchange/Runtime/Source/Parsers/Fbx/Private/FbxAPI.h b/Engine/Plugins/Interchange/Runtime/Source/Parsers/Fbx/Private/FbxAPI.h
--- a/Engine/Plugins/Interchange/Runtime/Source/Parsers/Fbx/Private/FbxAPI.h
+++ b/Engine/Plugins/Interchange/Runtime/Source/Parsers/Fbx/Private/FbxAPI.h
@@ -42,11 +42,20 @@ namespace UE
 
 				~FFbxParser();
 
+				void Reset();
+
 				void SetResultContainer(UInterchangeResultsContainer* Result)
 				{
 					ResultsContainer = Result;
 				}
 
+				void SetConvertSettings(const bool InbConvertScene, const bool InbForceFrontXAxis, const bool InbConvertSceneUnit)
+				{
+					bConvertScene = InbConvertScene;
+					bForceFrontXAxis = InbForceFrontXAxis;
+					bConvertSceneUnit = InbConvertSceneUnit;
+				}
+
 				//return the fbx helper for this parser
 				const TSharedPtr<FFbxHelper> GetFbxHelper();
 
@@ -113,6 +122,11 @@ namespace UE
 
 				//For PivotReset and Animation Conversion:
 				double FrameRate = 30.0;
+
+				//Convert settings
+				bool bConvertScene = true;
+				bool bForceFrontXAxis = false;
+				bool bConvertSceneUnit = true;
 			};
 		}//ns Private
 	}//ns Interchange
diff --git a/Engine/Plugins/Interchange/Runtime/Source/Parsers/Fbx/Private/FbxConvert.cpp b/Engine/Plugins/Interchange/Runtime/Source/Parsers/Fbx/Private/FbxConvert.cpp
--- a/Engine/Plugins/Interchange/Runtime/Source/Parsers/Fbx/Private/FbxConvert.cpp
+++ b/Engine/Plugins/Interchange/Runtime/Source/Parsers/Fbx/Private/FbxConvert.cpp
@@ -13,7 +13,7 @@ namespace UE
 	{
 		namespace Private
 		{
-			void FFbxConvert::ConvertScene(FbxScene* SDKScene)
+			void FFbxConvert::ConvertScene(FbxScene* SDKScene, const bool bConvertScene, const bool bForceFrontXAxis, const bool bConvertSceneUnit)
 			{
 				if (!ensure(SDKScene))
 				{
@@ -42,30 +42,33 @@ namespace UE
 					}
 				}
 
-				//Set the original file information
-				FbxAxisSystem FileAxisSystem = SDKScene->GetGlobalSettings().GetAxisSystem();
-				FbxSystemUnit FileUnitSystem = SDKScene->GetGlobalSettings().GetSystemUnit();
 
+				if (bConvertScene)
+				{
+					//Set the original file information
+					FbxAxisSystem FileAxisSystem = SDKScene->GetGlobalSettings().GetAxisSystem();
 
-				//UE is: z up, front x, left handed
-				FbxAxisSystem::EUpVector UpVector = FbxAxisSystem::EUpVector::eZAxis;
-#if CONVERT_TO_FRONT_X
-				FbxAxisSystem::EFrontVector FrontVector = (FbxAxisSystem::EFrontVector)FbxAxisSystem::eParityEven;
-#else
-				FbxAxisSystem::EFrontVector FrontVector = (FbxAxisSystem::EFrontVector) - FbxAxisSystem::eParityOdd;
-#endif //CONVERT_TO_FRONT_X
-				FbxAxisSystem::ECoordSystem CoordSystem = FbxAxisSystem::ECoordSystem::eRightHanded;
-				FbxAxisSystem UnrealImportAxis(UpVector, FrontVector, CoordSystem);
 
-				if (FileAxisSystem != UnrealImportAxis)
-				{
-					FbxRootNodeUtility::RemoveAllFbxRoots(SDKScene);
-					UnrealImportAxis.ConvertScene(SDKScene);
+					//UE is: z up, front x, left handed
+					FbxAxisSystem::EUpVector UpVector = FbxAxisSystem::EUpVector::eZAxis;
+					FbxAxisSystem::EFrontVector FrontVector = (FbxAxisSystem::EFrontVector)(bForceFrontXAxis ? FbxAxisSystem::eParityEven : -FbxAxisSystem::eParityOdd);
+					FbxAxisSystem::ECoordSystem CoordSystem = FbxAxisSystem::ECoordSystem::eRightHanded;
+					FbxAxisSystem UnrealImportAxis(UpVector, FrontVector, CoordSystem);
+
+					if (FileAxisSystem != UnrealImportAxis)
+					{
+						FbxRootNodeUtility::RemoveAllFbxRoots(SDKScene);
+						UnrealImportAxis.ConvertScene(SDKScene);
+					}
 				}
 
-				if (FileUnitSystem != FbxSystemUnit::cm)
+				if (bConvertSceneUnit)
 				{
-					FbxSystemUnit::cm.ConvertScene(SDKScene);
+					FbxSystemUnit FileUnitSystem = SDKScene->GetGlobalSettings().GetSystemUnit();
+					if (FileUnitSystem != FbxSystemUnit::cm)
+					{
+						FbxSystemUnit::cm.ConvertScene(SDKScene);
+					}
 				}
 
 				//Reset all the transform evaluation cache since we change some node transform
diff --git a/Engine/Plugins/Interchange/Runtime/Source/Parsers/Fbx/Private/FbxConvert.h b/Engine/Plugins/Interchange/Runtime/Source/Parsers/Fbx/Private/FbxConvert.h
--- a/Engine/Plugins/Interchange/Runtime/Source/Parsers/Fbx/Private/FbxConvert.h
+++ b/Engine/Plugins/Interchange/Runtime/Source/Parsers/Fbx/Private/FbxConvert.h
@@ -156,7 +156,7 @@ namespace UE
 				/**
 				 * Convert a fbx scene 
 				 */
-				static void ConvertScene(FbxScene* SDKScene);
+				static void ConvertScene(FbxScene* SDKScene, const bool bConvertScene, const bool bForceFrontXAxis, const bool bConvertSceneUnit);
 				
 				/** Scene Conversion API End */
 				//////////////////////////////////////////////////////////////////////////
diff --git a/Engine/Plugins/Interchange/Runtime/Source/Parsers/Fbx/Private/FbxHelper.h b/Engine/Plugins/Interchange/Runtime/Source/Parsers/Fbx/Private/FbxHelper.h
--- a/Engine/Plugins/Interchange/Runtime/Source/Parsers/Fbx/Private/FbxHelper.h
+++ b/Engine/Plugins/Interchange/Runtime/Source/Parsers/Fbx/Private/FbxHelper.h
@@ -32,7 +32,7 @@ namespace UE::Interchange
 		struct FFbxHelper
 		{
 		public:
-
+			void Reset() { MaterialNameClashMap.Reset(); }
 			FString GetMeshName(FbxGeometryBase* Mesh) const;
 			FString GetMeshUniqueID(FbxGeometryBase* Mesh) const;
 
diff --git a/Engine/Plugins/Interchange/Runtime/Source/Parsers/Fbx/Private/InterchangeFbxParser.cpp b/Engine/Plugins/Interchange/Runtime/Source/Parsers/Fbx/Private/InterchangeFbxParser.cpp
--- a/Engine/Plugins/Interchange/Runtime/Source/Parsers/Fbx/Private/InterchangeFbxParser.cpp
+++ b/Engine/Plugins/Interchange/Runtime/Source/Parsers/Fbx/Private/InterchangeFbxParser.cpp
@@ -55,12 +55,23 @@ namespace UE::Interchange
 		FbxParserPrivate = nullptr;
 	}
 
+	void FInterchangeFbxParser::Reset()
+	{
+		ResultPayloads.Reset();
+		FbxParserPrivate->Reset();
+	}
+
 	void FInterchangeFbxParser::SetResultContainer(UInterchangeResultsContainer* Result)
 	{
 		InternalResultsContainer = Result;
 		FbxParserPrivate->SetResultContainer(Result);
 	}
 
+	void FInterchangeFbxParser::SetConvertSettings(const bool InbConvertScene, const bool InbForceFrontXAxis, const bool InbConvertSceneUnit)
+	{
+		FbxParserPrivate->SetConvertSettings(InbConvertScene, InbForceFrontXAxis, InbConvertSceneUnit);
+	}
+
 	void FInterchangeFbxParser::LoadFbxFile(const FString& Filename, const FString& ResultFolder)
 	{
 		check(FbxParserPrivate.IsValid());
diff --git a/Engine/Plugins/Interchange/Runtime/Source/Parsers/Fbx/Public/InterchangeFbxParser.h b/Engine/Plugins/Interchange/Runtime/Source/Parsers/Fbx/Public/InterchangeFbxParser.h
--- a/Engine/Plugins/Interchange/Runtime/Source/Parsers/Fbx/Public/InterchangeFbxParser.h
+++ b/Engine/Plugins/Interchange/Runtime/Source/Parsers/Fbx/Public/InterchangeFbxParser.h
@@ -29,7 +29,11 @@ namespace UE
 			
 			void ReleaseResources();
 
+			void Reset();
+
 			void SetResultContainer(UInterchangeResultsContainer* Result);
+
+			void SetConvertSettings(const bool InbConvertScene, const bool InbForceFrontXAxis, const bool InbConvertSceneUnit);
 			/**
 			 * Parse a file support by the fbx sdk. It just extract all the fbx node and create a FBaseNodeContainer and dump it in a json file inside the ResultFolder
 			 * @param - Filename is the file that the fbx sdk will read (.fbx or .obj)
diff --git a/Engine/Source/Programs/InterchangeWorker/Private/InterchangeWorkerImpl.cpp b/Engine/Source/Programs/InterchangeWorker/Private/InterchangeWorkerImpl.cpp
--- a/Engine/Source/Programs/InterchangeWorker/Private/InterchangeWorkerImpl.cpp
+++ b/Engine/Source/Programs/InterchangeWorker/Private/InterchangeWorkerImpl.cpp
@@ -290,6 +290,8 @@ ETaskState FInterchangeWorkerImpl::LoadFbxFile(const FJsonLoadSourceCmd& LoadSou
 {
 	ETaskState ResultState = ETaskState::Unknown;
 	FString SourceFilename = LoadSourceCommand.GetSourceFilename();
+	FbxParser.Reset();
+	FbxParser.SetConvertSettings(LoadSourceCommand.GetDoesConvertScene(), LoadSourceCommand.GetDoesForceFrontXAxis(), LoadSourceCommand.GetDoesConvertSceneUnit());
 	FbxParser.LoadFbxFile(SourceFilename, ResultFolder);
 	FJsonLoadSourceCmd::JsonResultParser ResultParser;
 	ResultParser.SetResultFilename(FbxParser.GetResultFilepath());
diff --git a/Engine/Source/Runtime/Interchange/Core/Private/InterchangeTranslatorBase.cpp b/Engine/Source/Runtime/Interchange/Core/Private/InterchangeTranslatorBase.cpp
--- a/Engine/Source/Runtime/Interchange/Core/Private/InterchangeTranslatorBase.cpp
+++ b/Engine/Source/Runtime/Interchange/Core/Private/InterchangeTranslatorBase.cpp
@@ -3,9 +3,159 @@
 #include "InterchangeTranslatorBase.h"
 
 #include "Algo/Find.h"
+#include "InterchangeLogPrivate.h"
+#include "Misc/ConfigCacheIni.h"
 
 #include UE_INLINE_GENERATED_CPP_BY_NAME(InterchangeTranslatorBase)
 
+namespace UE::Interchange::TranslatorPrivate
+{
+	FString CreateConfigSectionName(UClass* TranslatorClass)
+	{
+		FString Section = TEXT("Interchange_Translator_ClassName_") + TranslatorClass->GetName();
+		return Section;
+	}
+}
+
+void UInterchangeTranslatorSettings::LoadSettings()
+{
+	const FString& ConfigFilename = GEditorPerProjectIni;
+	int32 PortFlags = 0;
+	UClass* Class = this->GetClass();
+	FString SectionName = UE::Interchange::TranslatorPrivate::CreateConfigSectionName(Class);
+	for (FProperty* Property = Class->PropertyLink; Property; Property = Property->PropertyLinkNext)
+	{
+		//Do not load a transient property
+		if (Property->HasAnyPropertyFlags(CPF_Transient))
+		{
+			continue;
+		}
+		FString Key = Property->GetName();
+		FArrayProperty* Array = CastField<FArrayProperty>(Property);
+		if (Array)
+		{
+			const FConfigSection* Section = GConfig->GetSection(*SectionName, false/*bForce*/, ConfigFilename);
+			if (Section != nullptr)
+			{
+				TArray<FConfigValue> List;
+				const FName KeyName(*Key, FNAME_Find);
+				Section->MultiFind(KeyName, List);
+
+				FScriptArrayHelper_InContainer ArrayHelper(Array, this);
+				// Only override default properties if there is something to override them with.
+				if (List.Num() > 0)
+				{
+					ArrayHelper.EmptyAndAddValues(List.Num());
+					for (int32 i = List.Num() - 1, c = 0; i >= 0; i--, c++)
+					{
+						Array->Inner->ImportText_Direct(*List[i].GetValue(), ArrayHelper.GetRawPtr(c), this, PortFlags);
+					}
+				}
+				else
+				{
+					int32 Index = 0;
+					const FConfigValue* ElementValue = nullptr;
+					do
+					{
+						// Add array index number to end of key
+						FString IndexedKey = FString::Printf(TEXT("%s[%i]"), *Key, Index);
+
+						// Try to find value of key
+						const FName IndexedName(*IndexedKey, FNAME_Find);
+						if (IndexedName == NAME_None)
+						{
+							break;
+						}
+						ElementValue = Section->Find(IndexedName);
+
+						// If found, import the element
+						if (ElementValue != nullptr)
+						{
+							// expand the array if necessary so that Index is a valid element
+							ArrayHelper.ExpandForIndex(Index);
+							Array->Inner->ImportText_Direct(*ElementValue->GetValue(), ArrayHelper.GetRawPtr(Index), this, PortFlags);
+						}
+
+						Index++;
+					} while (ElementValue || Index < ArrayHelper.Num());
+				}
+			}
+		}
+		else
+		{
+			for (int32 i = 0; i < Property->ArrayDim; i++)
+			{
+				if (Property->ArrayDim != 1)
+				{
+					Key = FString::Printf(TEXT("%s[%i]"), *Property->GetName(), i);
+				}
+
+				FString Value;
+				bool bFoundValue = GConfig->GetString(*SectionName, *Key, Value, ConfigFilename);
+
+				if (bFoundValue)
+				{
+					if (Property->ImportText_Direct(*Value, Property->ContainerPtrToValuePtr<uint8>(this, i), this, PortFlags) == NULL)
+					{
+						// this should be an error as the properties from the .ini / .int file are not correctly being read in and probably are affecting things in subtle ways
+						UE_LOG(LogInterchangeCore, Error, TEXT("UInterchangeTranslatorSettings (class:%s) failed to load settings. Property: %s Value: %s"), *this->GetClass()->GetName(), *Property->GetName(), *Value);
+					}
+				}
+			}
+		}
+	}
+}
+
+void UInterchangeTranslatorSettings::SaveSettings()
+{
+	const FString& ConfigFilename = GEditorPerProjectIni;
+	int32 PortFlags = 0;
+	UClass* Class = this->GetClass();
+	FString SectionName = UE::Interchange::TranslatorPrivate::CreateConfigSectionName(Class);
+
+	for (FProperty* Property = Class->PropertyLink; Property; Property = Property->PropertyLinkNext)
+	{
+		//Do not save a transient property
+		if (Property->HasAnyPropertyFlags(CPF_Transient))
+		{
+			continue;
+		}
+
+		FString Key = Property->GetName();
+		FObjectProperty* SubObject = CastField<FObjectProperty>(Property);
+		FArrayProperty* Array = CastField<FArrayProperty>(Property);
+		if (Array)
+		{
+			GConfig->RemoveKeyFromSection(*SectionName, *Key, ConfigFilename);
+
+			FScriptArrayHelper_InContainer ArrayHelper(Array, this);
+			for (int32 i = 0; i < ArrayHelper.Num(); i++)
+			{
+				FString	Buffer;
+				Array->Inner->ExportTextItem_Direct(Buffer, ArrayHelper.GetRawPtr(i), ArrayHelper.GetRawPtr(i), this, PortFlags);
+				GConfig->AddToSection(*SectionName, *Key, Buffer, ConfigFilename);
+			}
+		}
+		else
+		{
+			TCHAR TempKey[MAX_SPRINTF] = TEXT("");
+			for (int32 Index = 0; Index < Property->ArrayDim; Index++)
+			{
+				if (Property->ArrayDim != 1)
+				{
+					FCString::Sprintf(TempKey, TEXT("%s[%i]"), *Property->GetName(), Index);
+					Key = TempKey;
+				}
+
+				FString	Value;
+				Property->ExportText_InContainer(Index, Value, this, this, this, PortFlags);
+				GConfig->SetString(*SectionName, *Key, *Value, ConfigFilename);
+			}
+		}
+	}
+	GConfig->Flush(0);
+}
+
 bool UInterchangeTranslatorBase::CanImportSourceData(const UInterchangeSourceData* InSourceData) const
 {
 	const bool bIncludeDot = false;
@@ -20,3 +170,4 @@ bool UInterchangeTranslatorBase::CanImportSourceData(const UInterchangeSourceDat
 
 	return bExtensionMatches;
 }
+
diff --git a/Engine/Source/Runtime/Interchange/Core/Public/InterchangeFactoryBase.h b/Engine/Source/Runtime/Interchange/Core/Public/InterchangeFactoryBase.h
--- a/Engine/Source/Runtime/Interchange/Core/Public/InterchangeFactoryBase.h
+++ b/Engine/Source/Runtime/Interchange/Core/Public/InterchangeFactoryBase.h
@@ -238,6 +238,7 @@ public:
 		UInterchangeBaseNodeContainer* NodeContainer = nullptr;
 		TArray<UInterchangePipelineBase*> Pipelines;
 		TArray<UObject*> OriginalPipelines;
+		UInterchangeTranslatorBase* Translator;
 		
 		bool bIsReimport  = false;
 ;
diff --git a/Engine/Source/Runtime/Interchange/Core/Public/InterchangeTranslatorBase.h b/Engine/Source/Runtime/Interchange/Core/Public/InterchangeTranslatorBase.h
--- a/Engine/Source/Runtime/Interchange/Core/Public/InterchangeTranslatorBase.h
+++ b/Engine/Source/Runtime/Interchange/Core/Public/InterchangeTranslatorBase.h
@@ -20,6 +20,20 @@
 
 class UInterchangeBaseNodeContainer;
 
+/**
+ * Base class for translator settings, the UInterchangeAssetImportData will store this data if the source translator exist
+ */
+UCLASS(BlueprintType, Blueprintable, editinlinenew, Abstract, MinimalAPI)
+class UInterchangeTranslatorSettings : public UObject
+{
+	GENERATED_BODY()
+
+public:
+	INTERCHANGECORE_API void LoadSettings();
+
+	INTERCHANGECORE_API void SaveSettings();
+};
+
 UENUM(BlueprintType)
 enum class EInterchangeTranslatorType : uint8
 {
@@ -135,6 +149,15 @@ public:
 		return SourceData;
 	}
 
+	/**
+	 * Return this translator settings
+	 */
+	INTERCHANGECORE_API virtual UInterchangeTranslatorSettings* GetSettings() const { return nullptr; }
+
+	/**
+	 * Reimport will set the settings if they exist in UInterchangeAssetImportData
+	 */
+	INTERCHANGECORE_API virtual void SetSettings(const UInterchangeTranslatorSettings* InterchangeTranslatorSettings) { }
 
 	UPROPERTY()
 	TObjectPtr<UInterchangeResultsContainer> Results;
diff --git a/Engine/Source/Runtime/Interchange/Core/Public/Nodes/InterchangeBaseNodeContainer.h b/Engine/Source/Runtime/Interchange/Core/Public/Nodes/InterchangeBaseNodeContainer.h
--- a/Engine/Source/Runtime/Interchange/Core/Public/Nodes/InterchangeBaseNodeContainer.h
+++ b/Engine/Source/Runtime/Interchange/Core/Public/Nodes/InterchangeBaseNodeContainer.h
@@ -40,6 +40,20 @@ class UInterchangeBaseNodeContainer : public UObject
 public:
 	INTERCHANGECORE_API UInterchangeBaseNodeContainer();
 
+	/**
+	 * Empty the container.
+	 *
+	 * @param Node - a pointer on the node you want to add
+	 * @return: return the node unique ID of the added item. If the node already exist it will return the existing ID. Return InvalidNodeUid if the node cannot be added.
+	 *
+	 */
+	UFUNCTION(BlueprintCallable, Category = "Interchange | Node Container")
+	INTERCHANGECORE_API void Reset()
+	{
+		Nodes.Reset();
+		ChildrenCache.Reset();
+	}
+
 	/**
 	 * Add a node in the container, the node will be add into a TMap.
 	 *
diff --git a/Engine/Source/Runtime/Interchange/Engine/Private/InterchangeAssetImportData.cpp b/Engine/Source/Runtime/Interchange/Engine/Private/InterchangeAssetImportData.cpp
--- a/Engine/Source/Runtime/Interchange/Engine/Private/InterchangeAssetImportData.cpp
+++ b/Engine/Source/Runtime/Interchange/Engine/Private/InterchangeAssetImportData.cpp
@@ -46,6 +46,45 @@ void UInterchangeAssetImportData::PostLoad()
 	}
 }
 
+UObject* DeSerializeTranslatorSettings(const FString& TranslatorSettingsStr, UClass* TranslatorSettingsClass)
+{
+	UInterchangeTranslatorSettings* GeneratedTranslatorSettings = NewObject<UInterchangeTranslatorSettings>(GetTransientPackage(), TranslatorSettingsClass);
+	GeneratedTranslatorSettings->ClearInternalFlags(EInternalObjectFlags::Async);
+
+	TSharedPtr<FJsonObject> RootObject;
+	TSharedRef< TJsonReader<> > JsonReader = TJsonReaderFactory<>::Create(TranslatorSettingsStr);
+	if (FJsonSerializer::Deserialize(JsonReader, RootObject))
+	{
+		const TSharedPtr<FJsonObject> JsonTranslatorSettingsProperties = RootObject->GetObjectField(TEXT("GeneratedTranslatorSettings"));
+		FJsonObjectConverter::JsonObjectToUStruct(JsonTranslatorSettingsProperties.ToSharedRef(), TranslatorSettingsClass, GeneratedTranslatorSettings, 0, 0);
+	}
+
+	return GeneratedTranslatorSettings;
+}
+
+FString SerializeTranslatorSettings(UObject* TranslatorSettings)
+{
+	if (UClass* TranslatorSettingsClass = TranslatorSettings->GetClass())
+	{
+		TSharedRef<FJsonObject> RootObject = MakeShareable(new FJsonObject);
+		TSharedRef<FJsonObject> TranslatorSettingsPropertiesObject = MakeShareable(new FJsonObject);
+		if (FJsonObjectConverter::UStructToJsonObject(TranslatorSettingsClass, TranslatorSettings, TranslatorSettingsPropertiesObject, 0, 0))
+		{
+			RootObject->SetField(TEXT("GeneratedTranslatorSettings"), MakeShareable(new FJsonValueObject(TranslatorSettingsPropertiesObject)));
+		}
+		//Write the json file
+		FString Json;
+		TSharedRef<TJsonWriter<> > JsonWriter = TJsonWriterFactory<>::Create(&Json, 0);
+		if (FJsonSerializer::Serialize(RootObject, JsonWriter))
+		{
+			return Json;
+		}
+	}
+
+	return TEXT("");
+}
+
+
 UObject* DeSerializePipeline(const FString& PipelineStr, UClass* PipelineClass)
 {
 	UInterchangePipelineBase* GeneratedPipeline = NewObject<UInterchangePipelineBase>(GetTransientPackage(), PipelineClass);
@@ -124,6 +163,26 @@ void UInterchangeAssetImportData::SetNodeContainer(UInterchangeBaseNodeContainer
 	}
 }
 
+const UInterchangeTranslatorSettings* UInterchangeAssetImportData::GetTranslatorSettings() const
+{
+	ProcessTranslatorCache();
+	return TransientTranslatorSettings;
+}
+
+void UInterchangeAssetImportData::SetTranslatorSettings(UInterchangeTranslatorSettings* TranslatorSettings) const
+{
+	TransientTranslatorSettings = TranslatorSettings;
+
+	//Serialize cache
+	CachedTranslatorSettings = {};
+	if (TranslatorSettings)
+	{
+		FString TranslatorSettingsJSON = SerializePipeline(TranslatorSettings);
+
+		FString TranslatorSettingsClassFullName = TranslatorSettings->GetClass()->GetFullName();
+		CachedTranslatorSettings = TPair<FString, FString>(TranslatorSettingsClassFullName, TranslatorSettingsJSON);
+	}
+}
 
 void UInterchangeAssetImportData::SetPipelines(const TArray<UObject*>& InPipelines) const
 {
@@ -217,6 +276,43 @@ void UInterchangeAssetImportData::ProcessContainerCache() const
 	}
 }
 
+void UInterchangeAssetImportData::ProcessTranslatorCache() const
+{
+	if (UInterchangeManager::IsInterchangeImportEnabled())
+	{
+		//de-serialize
+		if (!TransientTranslatorSettings && !CachedTranslatorSettings.Key.IsEmpty())
+		{
+			TMap<FString, UClass*> ClassPerName;
+			for (FThreadSafeObjectIterator It(UClass::StaticClass()); It; ++It)
+			{
+				UClass* Class = Cast<UClass>(*It);
+				if (Class->IsChildOf(UInterchangeTranslatorSettings::StaticClass()))
+				{
+					ClassPerName.Add(Class->GetFullName(), Class);
+				}
+			}
+
+			FString ClassFullName = CachedTranslatorSettings.Key;
+			FCoreRedirectObjectName RedirectedObjectName = FCoreRedirects::GetRedirectedName(ECoreRedirectFlags::Type_Class, FCoreRedirectObjectName(ClassFullName));
+			if (RedirectedObjectName.IsValid())
+			{
+				ClassFullName = RedirectedObjectName.ToString();
+			}
+			//This cannot fail to make sure we have a healty serialization
+			if (!ensure(ClassPerName.Contains(ClassFullName)))
+			{
+				//We did not successfully serialize the content of the file into the node container
+				return;
+			}
+
+			UClass* ToCreateClass = ClassPerName.FindChecked(ClassFullName);
+
+			TransientTranslatorSettings = Cast<UInterchangeTranslatorSettings>(DeSerializeTranslatorSettings(CachedTranslatorSettings.Value, ToCreateClass));
+		}
+	}
+}
+
 void UInterchangeAssetImportData::ProcessPipelinesCache() const
 {
 	if (UInterchangeManager::IsInterchangeImportEnabled())
diff --git a/Engine/Source/Runtime/Interchange/Engine/Private/InterchangeImportCommon.cpp b/Engine/Source/Runtime/Interchange/Engine/Private/InterchangeImportCommon.cpp
--- a/Engine/Source/Runtime/Interchange/Engine/Private/InterchangeImportCommon.cpp
+++ b/Engine/Source/Runtime/Interchange/Engine/Private/InterchangeImportCommon.cpp
@@ -9,6 +9,7 @@
 #include "InterchangePythonPipelineBase.h"
 #include "InterchangeSceneImportAsset.h"
 #include "InterchangeSourceData.h"
+#include "InterchangeTranslatorBase.h"
 #include "Nodes/InterchangeBaseNodeContainer.h"
 #include "Nodes/InterchangeFactoryBaseNode.h"
 #include "UObject/Object.h"
@@ -79,6 +80,11 @@ namespace UE::Interchange
 			}
 			AssetImportData->SetNodeContainer(FactoryNodeContainer);
 
+			if (UInterchangeTranslatorSettings* InterchangeTranslatorSettings = Parameters.Translator->GetSettings())
+			{
+				AssetImportData->SetTranslatorSettings(InterchangeTranslatorSettings);
+			}
+
 			TArray<UObject*> NewPipelines;
 			for (const UObject* Pipeline : Parameters.Pipelines)
 			{
@@ -114,18 +120,21 @@ namespace UE::Interchange
 																						, const UInterchangeSourceData* InSourceData
 																						, FString InNodeUniqueID
 																						, UInterchangeBaseNodeContainer* InNodeContainer
-																						, const TArray<UObject*>& InPipelines)
+																						, const TArray<UObject*>& InPipelines
+																						, const UInterchangeTranslatorBase* InTranslator)
 		: AssetImportDataOuter(InAssetImportDataOuter)
 		, AssetImportData(InAssetImportData)
 		, SourceData(InSourceData)
 		, NodeUniqueID(InNodeUniqueID)
 		, NodeContainer(InNodeContainer)
 		, Pipelines(InPipelines)
+		, Translator(InTranslator)
 	{
 		ensure(AssetImportDataOuter);
 		ensure(SourceData);
 		ensure(!NodeUniqueID.IsEmpty());
 		ensure(NodeContainer);
+		ensure(Translator);
 	}
 
 	UAssetImportData* FFactoryCommon::UpdateImportAssetData(FUpdateImportAssetDataParameters& Parameters)
@@ -186,13 +195,15 @@ namespace UE::Interchange
 																				, const UInterchangeSourceData* InSourceData
 																				, FString InNodeUniqueID
 																				, UInterchangeBaseNodeContainer* InNodeContainer
-																				, const TArray<UObject*>& InPipelines)
+																				, const TArray<UObject*>& InPipelines
+																				, const UInterchangeTranslatorBase* InTranslator)
 		: FUpdateImportAssetDataParameters(InAssetImportDataOuter
 			, InAssetImportData
 			, InSourceData
 			, InNodeUniqueID
 			, InNodeContainer
 			, InPipelines
+			, InTranslator
 			)
 		, SourceFiles()
 	{
diff --git a/Engine/Source/Runtime/Interchange/Engine/Private/InterchangeManager.cpp b/Engine/Source/Runtime/Interchange/Engine/Private/InterchangeManager.cpp
--- a/Engine/Source/Runtime/Interchange/Engine/Private/InterchangeManager.cpp
+++ b/Engine/Source/Runtime/Interchange/Engine/Private/InterchangeManager.cpp
@@ -1608,6 +1608,14 @@ UInterchangeManager::ImportInternal(const FString& ContentPath, const UInterchan
 	for (int32 SourceDataIndex = 0; SourceDataIndex < AsyncHelper->SourceDatas.Num(); ++SourceDataIndex)
 	{
 		AsyncTranslator = GetTranslatorForSourceData(AsyncHelper->SourceDatas[SourceDataIndex]);
+		if (bIsReimport)
+		{
+			//Set translator settings if we are doing a reimport
+			if (const UInterchangeTranslatorSettings* InterchangeTranslatorSettings = OriginalAssetImportData->GetTranslatorSettings())
+			{
+				AsyncTranslator->SetSettings(InterchangeTranslatorSettings);
+			}
+		}
 		ensure(AsyncHelper->Translators.Add(AsyncTranslator) == SourceDataIndex);
 	}
 
@@ -1813,7 +1821,7 @@ UInterchangeManager::ImportInternal(const FString& ContentPath, const UInterchan
 				TranslateSourceFile();
 				UInterchangeBaseNodeContainer* BaseNodeContainer = AsyncHelper->BaseNodeContainers[SourceIndex].Get();
 				//Show the dialog, a plugin should have registered this dialog. We use a plugin to be able to use editor code when doing UI
-				EInterchangePipelineConfigurationDialogResult DialogResult = RegisteredPipelineConfiguration->ScriptedShowReimportPipelineConfigurationDialog(PipelineStacks, OutPipelines, DuplicateSourceData, BaseNodeContainer, ImportAssetParameters.ReimportAsset);
+				EInterchangePipelineConfigurationDialogResult DialogResult = RegisteredPipelineConfiguration->ScriptedShowReimportPipelineConfigurationDialog(PipelineStacks, OutPipelines, DuplicateSourceData, AsyncTranslator, BaseNodeContainer, ImportAssetParameters.ReimportAsset);
 				if (DialogResult == EInterchangePipelineConfigurationDialogResult::Cancel)
 				{
 					bImportCanceled = true;
@@ -1843,8 +1851,8 @@ UInterchangeManager::ImportInternal(const FString& ContentPath, const UInterchan
 				UInterchangeBaseNodeContainer* BaseNodeContainer = AsyncHelper->BaseNodeContainers[SourceIndex].Get();
 				//Show the dialog, a plugin should have register this dialog. We use a plugin to be able to use editor code when doing UI
 				EInterchangePipelineConfigurationDialogResult DialogResult = bImportScene
-					? RegisteredPipelineConfiguration->ScriptedShowScenePipelineConfigurationDialog(PipelineStacks, OutPipelines, DuplicateSourceData, BaseNodeContainer)
-					: RegisteredPipelineConfiguration->ScriptedShowPipelineConfigurationDialog(PipelineStacks, OutPipelines, DuplicateSourceData, BaseNodeContainer);
+					? RegisteredPipelineConfiguration->ScriptedShowScenePipelineConfigurationDialog(PipelineStacks, OutPipelines, DuplicateSourceData, AsyncTranslator, BaseNodeContainer)
+					: RegisteredPipelineConfiguration->ScriptedShowPipelineConfigurationDialog(PipelineStacks, OutPipelines, DuplicateSourceData, AsyncTranslator, BaseNodeContainer);
 
 				if (DialogResult == EInterchangePipelineConfigurationDialogResult::Cancel)
 				{
diff --git a/Engine/Source/Runtime/Interchange/Engine/Private/Tasks/InterchangeTaskCompletion.cpp b/Engine/Source/Runtime/Interchange/Engine/Private/Tasks/InterchangeTaskCompletion.cpp
--- a/Engine/Source/Runtime/Interchange/Engine/Private/Tasks/InterchangeTaskCompletion.cpp
+++ b/Engine/Source/Runtime/Interchange/Engine/Private/Tasks/InterchangeTaskCompletion.cpp
@@ -61,6 +61,7 @@ void UE::Interchange::FTaskPreCompletion::DoTask(ENamedThreads::Type CurrentThre
 		Arguments.NodeContainer = AsyncHelper->BaseNodeContainers[SourceIndex].Get();
 		Arguments.Pipelines = AsyncHelper->Pipelines;
 		Arguments.OriginalPipelines = AsyncHelper->OriginalPipelines;
+		Arguments.Translator = AsyncHelper->Translators[SourceIndex];
 
 		//First iteration to call SetupObject_GameThread and pipeline ExecutePostFactoryPipeline
 		for (const FImportAsyncHelper::FImportedObjectInfo& ObjectInfo : ImportedObjects)
diff --git a/Engine/Source/Runtime/Interchange/Engine/Private/Tasks/InterchangeTaskPipeline.cpp b/Engine/Source/Runtime/Interchange/Engine/Private/Tasks/InterchangeTaskPipeline.cpp
--- a/Engine/Source/Runtime/Interchange/Engine/Private/Tasks/InterchangeTaskPipeline.cpp
+++ b/Engine/Source/Runtime/Interchange/Engine/Private/Tasks/InterchangeTaskPipeline.cpp
@@ -188,6 +188,7 @@ void UE::Interchange::FTaskPostImport::DoTask(ENamedThreads::Type CurrentThread,
 			Arguments.NodeContainer = AsyncHelper->BaseNodeContainers[SourceIndex].Get();
 			Arguments.Pipelines = AsyncHelper->Pipelines;
 			Arguments.OriginalPipelines = AsyncHelper->OriginalPipelines;
+			Arguments.Translator = AsyncHelper->Translators[SourceIndex];
 
 			for (const UE::Interchange::FImportAsyncHelper::FImportedObjectInfo& ObjectInfo : ImportedInfos)
 			{
diff --git a/Engine/Source/Runtime/Interchange/Engine/Public/InterchangeAssetImportData.h b/Engine/Source/Runtime/Interchange/Engine/Public/InterchangeAssetImportData.h
--- a/Engine/Source/Runtime/Interchange/Engine/Public/InterchangeAssetImportData.h
+++ b/Engine/Source/Runtime/Interchange/Engine/Public/InterchangeAssetImportData.h
@@ -4,6 +4,7 @@
 
 #include "CoreMinimal.h"
 #include "EditorFramework/AssetImportData.h"
+#include "InterchangeTranslatorBase.h"
 #include "Nodes/InterchangeBaseNode.h"
 #include "Nodes/InterchangeBaseNodeContainer.h"
 #include "UObject/AssetRegistryTagsContext.h"
@@ -159,6 +160,13 @@ public:
 	UFUNCTION(BlueprintCallable, Category = "Interchange | AssetImportData")
 	INTERCHANGEENGINE_API UInterchangeFactoryBaseNode* GetStoredFactoryNode(const FString& InNodeUniqueId) const;
 
+	UFUNCTION(BlueprintCallable, Category = "Interchange | AssetImportData")
+	INTERCHANGEENGINE_API const UInterchangeTranslatorSettings* GetTranslatorSettings() const;
+
+	UFUNCTION(BlueprintCallable, Category = "Interchange | AssetImportData")
+	INTERCHANGEENGINE_API void SetTranslatorSettings(UInterchangeTranslatorSettings* TranslatorSettings) const;
+
+
 private:
 	UPROPERTY(meta = (DeprecatedProperty, DeprecationMessage = "Use GetNodeContainer/SetNodeContainer instead."))
 	TObjectPtr<UInterchangeBaseNodeContainer> NodeContainer_DEPRECATED;
@@ -171,11 +179,16 @@ private:
 	UPROPERTY(Transient)
 	mutable TArray<TObjectPtr<UObject>> TransientPipelines;
 
+	UPROPERTY(Transient)
+	mutable TObjectPtr<UInterchangeTranslatorSettings> TransientTranslatorSettings;
+
 	void ProcessContainerCache() const;
 	void ProcessPipelinesCache() const;
 	void ProcessDeprecatedData() const;
+	void ProcessTranslatorCache() const;
 	mutable TArray64<uint8> CachedNodeContainer;
 	mutable TArray<TPair<FString, FString>> CachedPipelines; //Class, Data(serialized JSON) pair
+	mutable TPair<FString, FString> CachedTranslatorSettings;
 };
 
 /**
diff --git a/Engine/Source/Runtime/Interchange/Engine/Public/InterchangeImportCommon.h b/Engine/Source/Runtime/Interchange/Engine/Public/InterchangeImportCommon.h
--- a/Engine/Source/Runtime/Interchange/Engine/Public/InterchangeImportCommon.h
+++ b/Engine/Source/Runtime/Interchange/Engine/Public/InterchangeImportCommon.h
@@ -14,6 +14,7 @@ class UInterchangeBaseNodeContainer;
 class UInterchangeFactoryBaseNode;
 class UInterchangePipelineBase;
 class UInterchangeSourceData;
+class UInterchangeTranslatorBase;
 class UObject;
 template <typename FuncType> class TFunctionRef;
 
@@ -35,13 +36,15 @@ namespace UE::Interchange
 			FString NodeUniqueID;
 			UInterchangeBaseNodeContainer* NodeContainer;
 			const TArray<UObject*> Pipelines;
+			const UInterchangeTranslatorBase* Translator;
 
 			INTERCHANGEENGINE_API FUpdateImportAssetDataParameters(UObject* InAssetImportDataOuter
 																, UAssetImportData* InAssetImportData
 																, const UInterchangeSourceData* InSourceData
 																, FString InNodeUniqueID
 																, UInterchangeBaseNodeContainer* InNodeContainer
-																, const TArray<UObject*>& InPipelines);
+																, const TArray<UObject*>& InPipelines
+																, const UInterchangeTranslatorBase* InTranslator);
 		};
 
 		/**
@@ -70,7 +73,8 @@ namespace UE::Interchange
 				, const UInterchangeSourceData* InSourceData
 				, FString InNodeUniqueID
 				, UInterchangeBaseNodeContainer* InNodeContainer
-				, const TArray<UObject*>& InPipelines);
+				, const TArray<UObject*>& InPipelines
+				, const UInterchangeTranslatorBase* InTranslator);
 		};
 
 		/**
diff --git a/Engine/Source/Runtime/Interchange/Engine/Public/InterchangePipelineConfigurationBase.h b/Engine/Source/Runtime/Interchange/Engine/Public/InterchangePipelineConfigurationBase.h
--- a/Engine/Source/Runtime/Interchange/Engine/Public/InterchangePipelineConfigurationBase.h
+++ b/Engine/Source/Runtime/Interchange/Engine/Public/InterchangePipelineConfigurationBase.h
@@ -6,6 +6,7 @@
 #include "UObject/Object.h"
 #include "UObject/ObjectMacros.h"
 #include "InterchangeSourceData.h"
+#include "InterchangeTranslatorBase.h"
 #include "Nodes/InterchangeBaseNodeContainer.h"
 
 #include "InterchangePipelineConfigurationBase.generated.h"
@@ -44,16 +45,18 @@ public:
 	INTERCHANGEENGINE_API EInterchangePipelineConfigurationDialogResult ScriptedShowPipelineConfigurationDialog(TArray<FInterchangeStackInfo>& PipelineStacks
 		, TArray<UInterchangePipelineBase*>& OutPipelines
 		, UInterchangeSourceData* SourceData
+		, UInterchangeTranslatorBase* Translator
 		, UInterchangeBaseNodeContainer* BaseNodeContainer);
 
 	/** The default implementation (call if the blueprint do not have any implementation) will call the virtual ExecuteImportPipeline */
 	EInterchangePipelineConfigurationDialogResult ScriptedShowPipelineConfigurationDialog_Implementation(TArray<FInterchangeStackInfo>& PipelineStacks
 		, TArray<UInterchangePipelineBase*>& OutPipelines
 		, UInterchangeSourceData* SourceData
+		, UInterchangeTranslatorBase* Translator
 		, UInterchangeBaseNodeContainer* BaseNodeContainer)
 	{
 		//By default we call the virtual import pipeline execution
-		return ShowPipelineConfigurationDialog(PipelineStacks, OutPipelines, SourceData, BaseNodeContainer);
+		return ShowPipelineConfigurationDialog(PipelineStacks, OutPipelines, SourceData, Translator, BaseNodeContainer);
 	}
 
 	/**
@@ -63,16 +66,18 @@ public:
 	INTERCHANGEENGINE_API EInterchangePipelineConfigurationDialogResult ScriptedShowScenePipelineConfigurationDialog(TArray<FInterchangeStackInfo>& PipelineStacks
 		, TArray<UInterchangePipelineBase*>& OutPipelines
 		, UInterchangeSourceData* SourceData
+		, UInterchangeTranslatorBase* Translator
 		, UInterchangeBaseNodeContainer* BaseNodeContainer);
 
 	/** The default implementation (call if the blueprint do not have any implementation) will call the virtual ExecuteImportPipeline */
 	EInterchangePipelineConfigurationDialogResult ScriptedShowScenePipelineConfigurationDialog_Implementation(TArray<FInterchangeStackInfo>& PipelineStacks
 		, TArray<UInterchangePipelineBase*>& OutPipelines
 		, UInterchangeSourceData* SourceData
+		, UInterchangeTranslatorBase* Translator
 		, UInterchangeBaseNodeContainer* BaseNodeContainer)
 	{
 		//By default we call the virtual import pipeline execution
-		return ShowScenePipelineConfigurationDialog(PipelineStacks, OutPipelines, SourceData, BaseNodeContainer);
+		return ShowScenePipelineConfigurationDialog(PipelineStacks, OutPipelines, SourceData, Translator, BaseNodeContainer);
 	}
 
 	/**
@@ -82,6 +87,7 @@ public:
 	INTERCHANGEENGINE_API EInterchangePipelineConfigurationDialogResult ScriptedShowReimportPipelineConfigurationDialog(TArray<FInterchangeStackInfo>& PipelineStacks
 		, TArray<UInterchangePipelineBase*>& OutPipelines
 		, UInterchangeSourceData* SourceData
+		, UInterchangeTranslatorBase* Translator
 		, UInterchangeBaseNodeContainer* BaseNodeContainer
 		, UObject* ReimportAsset);
 
@@ -89,11 +95,12 @@ public:
 	EInterchangePipelineConfigurationDialogResult ScriptedShowReimportPipelineConfigurationDialog_Implementation(TArray<FInterchangeStackInfo>& PipelineStacks
 		, TArray<UInterchangePipelineBase*>& OutPipelines
 		, UInterchangeSourceData* SourceData
+		, UInterchangeTranslatorBase* Translator
 		, UInterchangeBaseNodeContainer* BaseNodeContainer
 		, UObject* ReimportAsset)
 	{
 		//By default we call the virtual import pipeline execution
-		return ShowReimportPipelineConfigurationDialog(PipelineStacks, OutPipelines, SourceData, BaseNodeContainer, ReimportAsset);
+		return ShowReimportPipelineConfigurationDialog(PipelineStacks, OutPipelines, SourceData, Translator, BaseNodeContainer, ReimportAsset);
 	}
 
 protected:
@@ -104,6 +111,7 @@ protected:
 	virtual EInterchangePipelineConfigurationDialogResult ShowPipelineConfigurationDialog(TArray<FInterchangeStackInfo>& PipelineStacks
 		, TArray<UInterchangePipelineBase*>& OutPipelines
 		, TWeakObjectPtr<UInterchangeSourceData> SourceData
+		, TWeakObjectPtr <UInterchangeTranslatorBase> Translator
 		, TWeakObjectPtr <UInterchangeBaseNodeContainer> BaseNodeContainer)
 	{ 
 		//Not implemented
@@ -116,6 +124,7 @@ protected:
 	virtual EInterchangePipelineConfigurationDialogResult ShowScenePipelineConfigurationDialog(TArray<FInterchangeStackInfo>& PipelineStacks
 		, TArray<UInterchangePipelineBase*>& OutPipelines
 		, TWeakObjectPtr<UInterchangeSourceData> SourceData
+		, TWeakObjectPtr <UInterchangeTranslatorBase> Translator
 		, TWeakObjectPtr <UInterchangeBaseNodeContainer> BaseNodeContainer)
 	{ 
 		//Not implemented
@@ -128,6 +137,7 @@ protected:
 	virtual EInterchangePipelineConfigurationDialogResult ShowReimportPipelineConfigurationDialog(TArray<FInterchangeStackInfo>& PipelineStacks
 		, TArray<UInterchangePipelineBase*>& OutPipelines
 		, TWeakObjectPtr<UInterchangeSourceData> SourceData
+		, TWeakObjectPtr <UInterchangeTranslatorBase> Translator
 		, TWeakObjectPtr <UInterchangeBaseNodeContainer> BaseNodeContainer
 		, TWeakObjectPtr <UObject> ReimportAsset)
 	{
