Interchange: Beautify pipeline name in the import dialog

#rnx

#virtualized

diff --git a/Engine/Plugins/Interchange/Editor/Source/Pipelines/Private/InterchangeEditorPipelineDetails.cpp b/Engine/Plugins/Interchange/Editor/Source/Pipelines/Private/InterchangeEditorPipelineDetails.cpp
--- a/Engine/Plugins/Interchange/Editor/Source/Pipelines/Private/InterchangeEditorPipelineDetails.cpp
+++ b/Engine/Plugins/Interchange/Editor/Source/Pipelines/Private/InterchangeEditorPipelineDetails.cpp
@@ -323,6 +323,15 @@ void FInterchangePipelineBaseDetailsCustomization::CustomizeDetails(IDetailLayou
 				}
 			}
 			
+			const bool PipelineInternalEditionData = PropertyHandle->GetBoolMetaData(FName("PipelineInternalEditionData"));
+
+			//Hide property that are not suppose to show in the import dialog
+			if (!bAllowPropertyStatesEdition && PipelineInternalEditionData)
+			{
+				CachedDetailBuilder->HideProperty(PropertyHandle);
+				continue;
+			}
+
 			FName PropertyPath = FName(PropertyPtr->GetPathName());
 			CachedDetailBuilder->HideProperty(PropertyHandle);
 
@@ -440,6 +449,7 @@ void FInterchangePipelineBaseDetailsCustomization::CustomizeDetails(IDetailLayou
 					.Padding(3.0f, 1.0f)
 					[
 						SNew(SCheckBox)
+						.Visibility(PipelineInternalEditionData ? EVisibility::Collapsed : EVisibility::All)
 						.CheckedImage(FAppStyle::Get().GetBrush("Icons.Lock"))
 						.CheckedHoveredImage(FAppStyle::Get().GetBrush("Icons.Lock"))
 						.CheckedPressedImage(FAppStyle::Get().GetBrush("Icons.Lock"))
@@ -489,6 +499,7 @@ void FInterchangePipelineBaseDetailsCustomization::CustomizeDetails(IDetailLayou
 					.VAlign(VAlign_Center)
 					[
 						SNew(STextBlock)
+						.Visibility(PipelineInternalEditionData ? EVisibility::Collapsed : EVisibility::All)
 						.Font(IDetailLayoutBuilder::GetDetailFont())
 						.Text(NSLOCTEXT("InterchangePipelineBaseDetails::CustomizeDetails", "ShowWhenBasicLayoutText", "Basic Layout"))
 					]
@@ -497,6 +508,7 @@ void FInterchangePipelineBaseDetailsCustomization::CustomizeDetails(IDetailLayou
 					.Padding(3.0f, 1.0f)
 					[
 						SNew(SCheckBox)
+						.Visibility(PipelineInternalEditionData ? EVisibility::Collapsed : EVisibility::All)
 						.CheckedImage(FAppStyle::Get().GetBrush("Icons.Hidden"))
 						.CheckedHoveredImage(FAppStyle::Get().GetBrush("Icons.Hidden"))
 						.CheckedPressedImage(FAppStyle::Get().GetBrush("Icons.Hidden"))
@@ -534,6 +546,7 @@ void FInterchangePipelineBaseDetailsCustomization::CustomizeDetails(IDetailLayou
 					.VAlign(VAlign_Center)
 					[
 						SNew(STextBlock)
+						.Visibility(PipelineInternalEditionData ? EVisibility::Collapsed : EVisibility::All)
 						.Font(IDetailLayoutBuilder::GetDetailFont())
 						.Text(NSLOCTEXT("InterchangePipelineBaseDetails::CustomizeDetails", "HiddenAtImportText", "Import"))
 					]
@@ -542,6 +555,7 @@ void FInterchangePipelineBaseDetailsCustomization::CustomizeDetails(IDetailLayou
 					.Padding(3.0f, 1.0f)
 					[
 						SNew(SCheckBox)
+						.Visibility(PipelineInternalEditionData ? EVisibility::Collapsed : EVisibility::All)
 						.CheckedImage(FAppStyle::Get().GetBrush("Icons.Hidden"))
 						.CheckedHoveredImage(FAppStyle::Get().GetBrush("Icons.Hidden"))
 						.CheckedPressedImage(FAppStyle::Get().GetBrush("Icons.Hidden"))
@@ -581,6 +595,7 @@ void FInterchangePipelineBaseDetailsCustomization::CustomizeDetails(IDetailLayou
 					.VAlign(VAlign_Center)
 					[
 						SNew(STextBlock)
+						.Visibility(PipelineInternalEditionData ? EVisibility::Collapsed : EVisibility::All)
 						.Font(IDetailLayoutBuilder::GetDetailFont())
 						.Text(NSLOCTEXT("InterchangePipelineBaseDetails::CustomizeDetails", "HiddenAtReimportText", "Reimport"))
 					]
@@ -589,6 +604,7 @@ void FInterchangePipelineBaseDetailsCustomization::CustomizeDetails(IDetailLayou
 					.Padding(3.0f, 1.0f)
 					[
 						SNew(SCheckBox)
+						.Visibility(PipelineInternalEditionData ? EVisibility::Collapsed : EVisibility::All)
 						.CheckedImage(FAppStyle::Get().GetBrush("Icons.Hidden"))
 						.CheckedHoveredImage(FAppStyle::Get().GetBrush("Icons.Hidden"))
 						.CheckedPressedImage(FAppStyle::Get().GetBrush("Icons.Hidden"))
diff --git a/Engine/Plugins/Interchange/Editor/Source/Pipelines/Private/SInterchangePipelineConfigurationDialog.cpp b/Engine/Plugins/Interchange/Editor/Source/Pipelines/Private/SInterchangePipelineConfigurationDialog.cpp
--- a/Engine/Plugins/Interchange/Editor/Source/Pipelines/Private/SInterchangePipelineConfigurationDialog.cpp
+++ b/Engine/Plugins/Interchange/Editor/Source/Pipelines/Private/SInterchangePipelineConfigurationDialog.cpp
@@ -56,7 +56,11 @@ void SInterchangePipelineItem::Construct(
 	ConflictsComboEntry = MakeShared<FString>(ConflictsComboEntryText.ToString());
 	if (PipelineElementPtr.Get())
 	{
-		FString PipelineNameString = FString::Printf(TEXT("%s (%s)"), *PipelineElement->DisplayName, *PipelineElementPtr->GetClass()->GetName());
+		FString PipelineNameString = PipelineElement->DisplayName;
+		if (!PipelineElement->bBasicLayout)
+		{
+			PipelineNameString += FString::Printf(TEXT(" (%s)"), *PipelineElementPtr->GetClass()->GetName());
+		}
 		PipelineName = FText::FromString(PipelineNameString);
 		PipelineElement->ConflictInfos = PipelineElementPtr->GetConflictInfos(PipelineElement->ReimportObject, PipelineElement->Container, PipelineElement->SourceData);
 		if (PipelineElement->ConflictInfos.Num() > 0)
@@ -197,7 +201,11 @@ FString SInterchangePipelineConfigurationDialog::GetPipelineDisplayName(const UI
 {
 	static int32 RightChopIndex = ReimportPipelinePrefix.Len();
 
-	FString PipelineDisplayName = Pipeline->GetName();
+	FString PipelineDisplayName = Pipeline->ScriptedGetPipelineDisplayName();
+	if(PipelineDisplayName.IsEmpty())
+	{
+		PipelineDisplayName = Pipeline->GetName();
+	}
 	if (PipelineDisplayName.StartsWith(ReimportPipelinePrefix))
 	{
 		PipelineDisplayName = PipelineDisplayName.RightChop(RightChopIndex);
@@ -350,7 +358,7 @@ TSharedRef<SBox> SInterchangePipelineConfigurationDialog::SpawnPipelineConfigura
 					{
 						GeneratedPipeline->FilterPropertiesFromTranslatedData(BaseNodeContainer.Get());
 					}
-					PipelineListViewItems.Add(MakeShareable(new FInterchangePipelineItemType{ GetPipelineDisplayName(DefaultPipeline), GeneratedPipeline, ReimportObject.Get(), BaseNodeContainer.Get(), SourceData.Get()}));
+					PipelineListViewItems.Add(MakeShareable(new FInterchangePipelineItemType{ GetPipelineDisplayName(DefaultPipeline), GeneratedPipeline, ReimportObject.Get(), BaseNodeContainer.Get(), SourceData.Get(), bBasicLayout }));
 				}
 			}
 			SelectedStack = StackNamePtr;
@@ -1014,7 +1022,7 @@ void SInterchangePipelineConfigurationDialog::RefreshStack(bool bStackSelectionC
 				{
 					GeneratedPipeline->FilterPropertiesFromTranslatedData(BaseNodeContainer.Get());
 				}
-				PipelineListViewItems.Add(MakeShareable(new FInterchangePipelineItemType{ GetPipelineDisplayName(DefaultPipeline), GeneratedPipeline, ReimportObject.Get(), BaseNodeContainer.Get(), SourceData.Get() }));
+				PipelineListViewItems.Add(MakeShareable(new FInterchangePipelineItemType{ GetPipelineDisplayName(DefaultPipeline), GeneratedPipeline, ReimportObject.Get(), BaseNodeContainer.Get(), SourceData.Get(), bBasicLayout }));
 			}
 		}
 	}
diff --git a/Engine/Plugins/Interchange/Editor/Source/Pipelines/Private/SInterchangePipelineConfigurationDialog.h b/Engine/Plugins/Interchange/Editor/Source/Pipelines/Private/SInterchangePipelineConfigurationDialog.h
--- a/Engine/Plugins/Interchange/Editor/Source/Pipelines/Private/SInterchangePipelineConfigurationDialog.h
+++ b/Engine/Plugins/Interchange/Editor/Source/Pipelines/Private/SInterchangePipelineConfigurationDialog.h
@@ -26,6 +26,7 @@ public:
 	UObject* ReimportObject = nullptr;
 	UInterchangeBaseNodeContainer* Container = nullptr;
 	UInterchangeSourceData* SourceData = nullptr;
+	bool bBasicLayout = false;
 	TArray<FInterchangeConflictInfo> ConflictInfos;
 };
 
diff --git a/Engine/Plugins/Interchange/Runtime/Source/Pipelines/Public/InterchangeGenericAssetsPipeline.h b/Engine/Plugins/Interchange/Runtime/Source/Pipelines/Public/InterchangeGenericAssetsPipeline.h
--- a/Engine/Plugins/Interchange/Runtime/Source/Pipelines/Public/InterchangeGenericAssetsPipeline.h
+++ b/Engine/Plugins/Interchange/Runtime/Source/Pipelines/Public/InterchangeGenericAssetsPipeline.h
@@ -44,6 +44,11 @@ public:
 	UInterchangeGenericAssetsPipeline();
 	
 	//////	COMMON_CATEGORY Properties //////
+	
+	/** The name of the pipeline that will be display in the import dialog. */
+	UPROPERTY(EditAnywhere, BlueprintReadWrite, Category = "Common", meta = (StandAlonePipelineProperty = "True", PipelineInternalEditionData = "True"))
+	FString PipelineDisplayName;
+
 	/* Allow user to choose the re-import strategy. */
 	UPROPERTY(EditAnywhere, BlueprintReadWrite, Category = "Common", meta = (AdjustPipelineAndRefreshDetailOnChange = "True"))
 	EReimportStrategyFlags ReimportStrategy = EReimportStrategyFlags::ApplyNoProperties;
diff --git a/Engine/Plugins/Interchange/Runtime/Source/Pipelines/Public/InterchangeGenericMaterialPipeline.h b/Engine/Plugins/Interchange/Runtime/Source/Pipelines/Public/InterchangeGenericMaterialPipeline.h
--- a/Engine/Plugins/Interchange/Runtime/Source/Pipelines/Public/InterchangeGenericMaterialPipeline.h
+++ b/Engine/Plugins/Interchange/Runtime/Source/Pipelines/Public/InterchangeGenericMaterialPipeline.h
@@ -59,6 +59,10 @@ class INTERCHANGEPIPELINES_API UInterchangeGenericMaterialPipeline : public UInt
 public:
 	UInterchangeGenericMaterialPipeline();
 
+	/** The name of the pipeline that will be display in the import dialog. */
+	UPROPERTY(EditAnywhere, BlueprintReadWrite, Category = "Materials", meta = (StandAlonePipelineProperty = "True", PipelineInternalEditionData = "True"))
+	FString PipelineDisplayName;
+
 	/** If enabled, imports the material assets found in the sources. */
 	UPROPERTY(EditAnywhere, BlueprintReadWrite, Category = "Materials")
 	bool bImportMaterials = true;
diff --git a/Engine/Plugins/Interchange/Runtime/Source/Pipelines/Public/InterchangeGenericTexturePipeline.h b/Engine/Plugins/Interchange/Runtime/Source/Pipelines/Public/InterchangeGenericTexturePipeline.h
--- a/Engine/Plugins/Interchange/Runtime/Source/Pipelines/Public/InterchangeGenericTexturePipeline.h
+++ b/Engine/Plugins/Interchange/Runtime/Source/Pipelines/Public/InterchangeGenericTexturePipeline.h
@@ -19,6 +19,10 @@ class INTERCHANGEPIPELINES_API UInterchangeGenericTexturePipeline : public UInte
 	GENERATED_BODY()
 
 public:
+	/** The name of the pipeline that will be display in the import dialog. */
+	UPROPERTY(EditAnywhere, BlueprintReadWrite, Category = "Textures", meta = (StandAlonePipelineProperty = "True", PipelineInternalEditionData = "True"))
+	FString PipelineDisplayName;
+
 	/** If enabled, imports the texture assets found in the sources. */
 	UPROPERTY(EditAnywhere, BlueprintReadWrite, Category = "Textures")
 	bool bImportTextures = true;
diff --git a/Engine/Source/Runtime/Interchange/Core/Private/InterchangePipelineBase.cpp b/Engine/Source/Runtime/Interchange/Core/Private/InterchangePipelineBase.cpp
--- a/Engine/Source/Runtime/Interchange/Core/Private/InterchangePipelineBase.cpp
+++ b/Engine/Source/Runtime/Interchange/Core/Private/InterchangePipelineBase.cpp
@@ -307,6 +307,35 @@ UInterchangePipelineBase* UInterchangePipelineBase::GetMostPipelineOuter() const
 	return Top;
 }
 
+FString UInterchangePipelineBase::GetPipelineDisplayName() const
+{
+	int32 PortFlags = 0;
+	UClass* Class = this->GetClass();
+	for (FProperty* Property = Class->PropertyLink; Property; Property = Property->PropertyLinkNext)
+	{
+		FStrProperty* StringProperty = CastField<FStrProperty>(Property);
+		if (!StringProperty)
+		{
+			continue;
+		}
+		const FName PropertyName = Property->GetFName();
+		if (PropertyName != FName("PipelineDisplayName"))
+		{
+			continue;
+		}
+		//We found the property
+		FString Value = StringProperty->GetPropertyValue_InContainer(this, 0);
+		if (!Value.IsEmpty())
+		{
+			return Value;
+		}
+		//Stop field iteration
+		break;
+	}
+	//Did not found a valid DisplayName property, return the name of the object
+	return GetName();
+}
+
 #if WITH_EDITOR
 
 void UInterchangePipelineBase::InternalToggleVisibilityPropertiesOfMetaDataValue(UInterchangePipelineBase* OuterMostPipeline, UInterchangePipelineBase* Pipeline, bool bDoTransientSubPipeline, const FString& MetaDataKey, const FString& MetaDataValue, const bool VisibilityState)
diff --git a/Engine/Source/Runtime/Interchange/Core/Public/InterchangePipelineBase.h b/Engine/Source/Runtime/Interchange/Core/Public/InterchangePipelineBase.h
--- a/Engine/Source/Runtime/Interchange/Core/Public/InterchangePipelineBase.h
+++ b/Engine/Source/Runtime/Interchange/Core/Public/InterchangePipelineBase.h
@@ -167,6 +167,29 @@ class UInterchangePipelineBase : public UObject
 
 public:
 
+	/**
+	 * This function is call when we want to list pipeline in the import dialog. If not override the default behavior of this function will search if
+	 * the pipeline have a FString UPROPERTY named "PipelineDisplayName" and return the property value. If there is no FString UPROPERTY call "PipelineDisplayName" it will
+	 * return the name of the pipeline asset (UObject::GetName).
+	 * 
+	 * When creating a pipeline (c++, python or blueprint) you can simply add a UPROPERTY name "PipelineDisplayName" to your pipeline, you do not need to override the function.
+	 * Use the same category has your other options and put it on the top.
+	 * The meta tag StandAlonePipelineProperty will hide your PROPERTY if your pipeline is a sub object of another pipeline when showing the import dialog.
+	 * The meta tag PipelineInternalEditionData make sure the property will be show only when we edit the pipeline object (hidden when showing the import dialog).
+	 * 
+	 *
+	 * UPROPERTY(EditAnywhere, BlueprintReadWrite, Category = "Textures", meta = (StandAlonePipelineProperty = "True", PipelineInternalEditionData = "True"))
+	 * FString PipelineDisplayName;
+	 */
+	UFUNCTION(BlueprintNativeEvent, BlueprintCallable, Category = "Interchange | Pipeline")
+	INTERCHANGECORE_API FString ScriptedGetPipelineDisplayName() const;
+	/** The default implementation (call if the blueprint do not have any implementation) will call the virtual ExecutePipeline */
+	FString ScriptedGetPipelineDisplayName_Implementation() const
+	{
+		//By default we call the virtual GetPipelineDisplayName
+		return GetPipelineDisplayName();
+	}
+
 	/**
 	 * ScriptedExecutePipeline, is call after the translation and before we parse the graph to call the factory.
 	 * This is where factory node should be created by the pipeline.
@@ -441,6 +464,8 @@ public:
 
 protected:
 
+	INTERCHANGECORE_API virtual FString GetPipelineDisplayName() const;
+
 	virtual void ExecutePipeline(UInterchangeBaseNodeContainer* BaseNodeContainer, const TArray<UInterchangeSourceData*>& SourceDatas, const FString& ContentBasePath)
 	{
 	}
