Interchange: Prevent re-import of legacy imported asset with interchange to create physics asset. fix crash if a translator in the project settings do not exist anymore.

#rnx

diff --git a/Engine/Plugins/Interchange/Editor/Source/InterchangeEditor/Private/InterchangeFbxAssetImportDataConverter.cpp b/Engine/Plugins/Interchange/Editor/Source/InterchangeEditor/Private/InterchangeFbxAssetImportDataConverter.cpp
--- a/Engine/Plugins/Interchange/Editor/Source/InterchangeEditor/Private/InterchangeFbxAssetImportDataConverter.cpp
+++ b/Engine/Plugins/Interchange/Editor/Source/InterchangeEditor/Private/InterchangeFbxAssetImportDataConverter.cpp
@@ -189,6 +189,7 @@ namespace UE::Interchange::Private
 
 		GenericAssetPipeline->CommonSkeletalMeshesAndAnimationsProperties->bImportMeshesInBoneHierarchy = SkeletalMeshImportData->bImportMeshesInBoneHierarchy;
 		GenericAssetPipeline->CommonMeshesProperties->bKeepSectionsSeparate = SkeletalMeshImportData->bKeepSectionsSeparate;
+		GenericAssetPipeline->MeshPipeline->bCreatePhysicsAsset = false;
 		GenericAssetPipeline->MeshPipeline->bImportMorphTargets = SkeletalMeshImportData->bImportMorphTargets;
 		GenericAssetPipeline->MeshPipeline->bImportVertexAttributes = SkeletalMeshImportData->bImportVertexAttributes;
 		GenericAssetPipeline->MeshPipeline->bUpdateSkeletonReferencePose = SkeletalMeshImportData->bUpdateSkeletonReferencePose;
diff --git a/Engine/Source/Runtime/Interchange/Engine/Private/InterchangeManager.cpp b/Engine/Source/Runtime/Interchange/Engine/Private/InterchangeManager.cpp
--- a/Engine/Source/Runtime/Interchange/Engine/Private/InterchangeManager.cpp
+++ b/Engine/Source/Runtime/Interchange/Engine/Private/InterchangeManager.cpp
@@ -1798,7 +1798,7 @@ UInterchangeManager::ImportInternal(const FString& ContentPath, const UInterchan
 				for (const FInterchangeTranslatorPipelines& TranslatorPipelines : PipelineStack.PerTranslatorPipelines)
 				{
 					const UClass* TranslatorClass = TranslatorPipelines.Translator.LoadSynchronous();
-					if (ScopedTranslator.GetTranslator()->IsA(TranslatorClass))
+					if (ScopedTranslator.GetTranslator() && ScopedTranslator.GetTranslator()->IsA(TranslatorClass))
 					{
 						Pipelines = &TranslatorPipelines.Pipelines;
 						break;
