Interchange: When re-importing an asset, preview should not show any other nodes with the enabled check box true.

#rnx

diff --git a/Engine/Plugins/Interchange/Editor/Source/Pipelines/Private/SInterchangePipelineConfigurationDialog.cpp b/Engine/Plugins/Interchange/Editor/Source/Pipelines/Private/SInterchangePipelineConfigurationDialog.cpp
--- a/Engine/Plugins/Interchange/Editor/Source/Pipelines/Private/SInterchangePipelineConfigurationDialog.cpp
+++ b/Engine/Plugins/Interchange/Editor/Source/Pipelines/Private/SInterchangePipelineConfigurationDialog.cpp
@@ -1146,10 +1146,18 @@ FReply SInterchangePipelineConfigurationDialog::OnPreviewImport() const
 		}
 	}
 
-	//Set all node in preview mode so hide the internal data attributes
-	DuplicateBaseNodeContainer->IterateNodesOfType<UInterchangeFactoryBaseNode>([](const FString& NodeUid, UInterchangeFactoryBaseNode* Node)
+	
+	DuplicateBaseNodeContainer->IterateNodesOfType<UInterchangeFactoryBaseNode>([ClosureReimportObject = ReimportObject](const FString& NodeUid, UInterchangeFactoryBaseNode* Node)
 		{
+
+			//Set all node in preview mode so hide the internal data attributes
 			Node->UserInterfaceContext = EInterchangeNodeUserInterfaceContext::Preview;
+
+			//If we reimport a specific object we want to disabled all factory nodes that are not supporting the reimport object class
+			if (ClosureReimportObject.IsValid())
+			{
+				Node->SetEnabled(ClosureReimportObject.Get()->IsA(Node->GetObjectClass()));
+			}
 		});
 
 	//Create and show the graph inspector UI dialog
