Improved visibility for conflicts on reimport

diff --git a/Engine/Plugins/Interchange/Editor/Source/Pipelines/Private/InterchangeEditorPipelineDetails.cpp b/Engine/Plugins/Interchange/Editor/Source/Pipelines/Private/InterchangeEditorPipelineDetails.cpp
--- a/Engine/Plugins/Interchange/Editor/Source/Pipelines/Private/InterchangeEditorPipelineDetails.cpp
+++ b/Engine/Plugins/Interchange/Editor/Source/Pipelines/Private/InterchangeEditorPipelineDetails.cpp
@@ -9,6 +9,7 @@
 #include "Nodes/InterchangeBaseNode.h"
 #include "ScopedTransaction.h"
 #include "Styling/StyleColors.h"
+#include "Widgets/Input/SButton.h"
 #include "Widgets/Input/SNumericEntryBox.h"
 #include "Widgets/Input/STextComboBox.h"
 #include "Widgets/Layout/SBox.h"
@@ -108,6 +109,13 @@ TSharedRef<IDetailCustomization> FInterchangePipelineBaseDetailsCustomization::M
 	return MakeShareable(new FInterchangePipelineBaseDetailsCustomization);
 }
 
+static TArray<TArray<FInterchangeConflictInfo>> ConflicInfosStack;
+void FInterchangePipelineBaseDetailsCustomization::SetConflictsInfo(TArray<FInterchangeConflictInfo>& ConflictInfos)
+{
+	ensure(ConflicInfosStack.Num() == 0);
+	ConflicInfosStack.Push(ConflictInfos);
+}
+
 void FInterchangePipelineBaseDetailsCustomization::SetTextComboBoxWidget(IDetailPropertyRow& PropertyRow, const TSharedPtr<IPropertyHandle>& Handle, const TArray<FString>& PossibleValues)
 {
 	if (!Handle.IsValid() || !Handle->IsValidHandle() || PossibleValues.Num() < 1)
@@ -251,6 +259,80 @@ void FInterchangePipelineBaseDetailsCustomization::OnTextComboBoxChanged(TShared
 	}
 }
 
+FReply FInterchangePipelineBaseDetailsCustomization::ShowConflictDialog(FInterchangeConflictInfo ConflictInfo)
+{
+	ConflictInfo.Pipeline->ShowConflictDialog(ConflictInfo.UniqueId);
+	return FReply::Handled();
+}
+
+void FInterchangePipelineBaseDetailsCustomization::AddConflictSection()
+{
+	if (ConflicInfosStack.Num() == 0)
+	{
+		return;
+	}
+
+	TArray<FInterchangeConflictInfo> ConflictInfos = ConflicInfosStack.Pop(false);
+	if (!InterchangePipeline->IsReimportContext() || ConflictInfos.Num() == 0)
+	{
+		return;
+	}
+
+	static const FName ConflictsName = TEXT("Conflicts");
+	static const FText ConflictsText = LOCTEXT("Conflicts_CategoryName", "Conflicts");
+	IDetailCategoryBuilder& CategoryBuilder = CachedDetailBuilder->EditCategory(ConflictsName, ConflictsText, ECategoryPriority::Important);
+
+	static const FText Conflict_ButtonText = LOCTEXT("Conflicts_ButtonShow", "Show Conflict");
+	static const FSlateBrush* ConflictBrush = FAppStyle::GetBrush("Icons.Error");
+
+	for (const FInterchangeConflictInfo& ConflictInfo : ConflictInfos)
+	{
+		const FText Conflict_IconTooltip = FText::FromString(ConflictInfo.Description);
+		const FText Conflict_NameContent = FText::FromString(ConflictInfo.DisplayName);
+		const FText Conflict_ButtonTooltip = FText::Format(LOCTEXT("Conflict_ButtonTooltip", "Show Conflict for {0}"), Conflict_NameContent);
+
+		CategoryBuilder.AddCustomRow(ConflictsText).WholeRowContent()
+			[
+				SNew(SHorizontalBox)
+					+ SHorizontalBox::Slot()
+					.VAlign(VAlign_Center)
+					.HAlign(HAlign_Center)
+					.AutoWidth()
+					.Padding(2.0f, 2.0f, 5.0f, 2.0f)
+					[
+						SNew(SImage)
+							.ToolTipText(Conflict_IconTooltip)
+							.Image(ConflictBrush)
+							.ColorAndOpacity(FStyleColors::Warning)
+					]
+					+ SHorizontalBox::Slot()
+					.FillWidth(1.0f)
+					.VAlign(VAlign_Center)
+					[
+						SNew(STextBlock)
+							.Text(Conflict_NameContent)
+							.Font(IDetailLayoutBuilder::GetDetailFont())
+							.ColorAndOpacity(FStyleColors::Warning)
+					]
+					+ SHorizontalBox::Slot()
+					.AutoWidth()
+					[
+						SNew(SButton)
+							.ToolTipText(Conflict_ButtonTooltip)
+							.OnClicked(this, &FInterchangePipelineBaseDetailsCustomization::ShowConflictDialog, ConflictInfo)
+							.Content()
+							[
+								SNew(STextBlock)
+									.Text(Conflict_ButtonText)
+									.Font(IDetailLayoutBuilder::GetDetailFont())
+							]
+					]
+
+			];
+	}
+}
+
+
 void FInterchangePipelineBaseDetailsCustomization::CustomizeDetails(IDetailLayoutBuilder& DetailBuilder)
 {
 	CachedDetailBuilder = &DetailBuilder;
@@ -277,6 +359,8 @@ void FInterchangePipelineBaseDetailsCustomization::CustomizeDetails(IDetailLayou
 	CachedDetailBuilder->GetCategoryNames(AllCategoryNames);
 	TMap<FName, TArray<FName>> PropertiesPerCategorys;
 	InternalGetPipelineProperties(InterchangePipeline.Get(), AllCategoryNames, PropertiesPerCategorys);
+
+	AddConflictSection();
 	
 	for (const TPair<FName, TArray<FName>>& CategoryAndProperties : PropertiesPerCategorys)
 	{
@@ -348,23 +432,23 @@ void FInterchangePipelineBaseDetailsCustomization::CustomizeDetails(IDetailLayou
 						
 						if (SubCategoryData.Equals(TEXT("Build")))
 						{
-							LocalizeSubCategoryName = NSLOCTEXT("InterchangePipelineBaseDetails::CustomizeDetails", "SubCategory_Build", "Build");
+							LocalizeSubCategoryName = LOCTEXT("SubCategory_Build", "Build");
 						}
 						else if (SubCategoryData.Equals(TEXT("Collision")))
 						{
-							LocalizeSubCategoryName = NSLOCTEXT("InterchangePipelineBaseDetails::CustomizeDetails", "SubCategory_Collision", "Collision");
+							LocalizeSubCategoryName = LOCTEXT("SubCategory_Collision", "Collision");
 						}
 						else if (SubCategoryData.Equals(TEXT("Actors properties")))
 						{
-							LocalizeSubCategoryName = NSLOCTEXT("InterchangePipelineBaseDetails::CustomizeDetails", "SubCategory_Actors_properties", "Actors properties");
+							LocalizeSubCategoryName = LOCTEXT("SubCategory_Actors_properties", "Actors properties");
 						}
 						else if (SubCategoryData.Equals(TEXT("Reimport Actors")))
 						{
-							LocalizeSubCategoryName = NSLOCTEXT("InterchangePipelineBaseDetails::CustomizeDetails", "SubCategory_Reimport_Actors", "Reimport Actors");
+							LocalizeSubCategoryName = LOCTEXT("SubCategory_Reimport_Actors", "Reimport Actors");
 						}
 						else if (SubCategoryData.Equals(TEXT("Reimport Assets")))
 						{
-							LocalizeSubCategoryName = NSLOCTEXT("InterchangePipelineBaseDetails::CustomizeDetails", "SubCategory_Reimport_Assets", "Reimport Assets");
+							LocalizeSubCategoryName = LOCTEXT("SubCategory_Reimport_Assets", "Reimport Assets");
 						}
 						SubCategoryGroups.Add(SubCategoryData, &(Category.AddGroup(FName(SubCategoryData), LocalizeSubCategoryName)));
 					}
@@ -456,14 +540,14 @@ void FInterchangePipelineBaseDetailsCustomization::CustomizeDetails(IDetailLayou
 						.UncheckedImage(FAppStyle::Get().GetBrush("Icons.Unlock"))
 						.UncheckedHoveredImage(FAppStyle::Get().GetBrush("Icons.Unlock"))
 						.UncheckedPressedImage(FAppStyle::Get().GetBrush("Icons.Unlock"))
-						.ToolTipText(NSLOCTEXT("InterchangePipelineBaseDetails::CustomizeDetails", "LockedTooltip", "If true this property will be readonly in the interchange import dialog."))
+						.ToolTipText(LOCTEXT("LockedTooltip", "If true this property will be readonly in the interchange import dialog."))
 						.OnCheckStateChanged_Lambda([InterchangePipelinePtr = InterchangePipeline, PropertyPath](ECheckBoxState CheckType)
 						{
 							if (!ensure(InterchangePipelinePtr.IsValid()))
 							{
 								return;
 							}
-							FScopedTransaction ScopedTransaction(NSLOCTEXT("InterchangePipelineBaseDetails::CustomizeDetails", "TransactionTogglePropertyLocked", "Toggle property locked at import"), !GIsTransacting);
+							FScopedTransaction ScopedTransaction(LOCTEXT("TransactionTogglePropertyLocked", "Toggle property locked at import"), !GIsTransacting);
 							InterchangePipelinePtr->Modify();
 							InterchangePipelinePtr->FindOrAddPropertyStates(PropertyPath).SetPropertyLocked(CheckType == ECheckBoxState::Checked);
 							InterchangePipelinePtr->PostEditChange();
@@ -501,7 +585,7 @@ void FInterchangePipelineBaseDetailsCustomization::CustomizeDetails(IDetailLayou
 						SNew(STextBlock)
 						.Visibility(PipelineInternalEditionData ? EVisibility::Collapsed : EVisibility::All)
 						.Font(IDetailLayoutBuilder::GetDetailFont())
-						.Text(NSLOCTEXT("InterchangePipelineBaseDetails::CustomizeDetails", "ShowWhenBasicLayoutText", "Basic Layout"))
+						.Text(LOCTEXT("ShowWhenBasicLayoutText", "Basic Layout"))
 					]
 					+ SHorizontalBox::Slot()
 					.AutoWidth()
@@ -515,14 +599,14 @@ void FInterchangePipelineBaseDetailsCustomization::CustomizeDetails(IDetailLayou
 						.UncheckedImage(FAppStyle::Get().GetBrush("Icons.Visible"))
 						.UncheckedHoveredImage(FAppStyle::Get().GetBrush("Icons.Visible"))
 						.UncheckedPressedImage(FAppStyle::Get().GetBrush("Icons.Visible"))
-						.ToolTipText(NSLOCTEXT("InterchangePipelineBaseDetails::CustomizeDetails", "VisibleTooltipBasicLayout", "If true this property will be visible when displaying the interchange import dialog with basic layout."))
+						.ToolTipText(LOCTEXT("VisibleTooltipBasicLayout", "If true this property will be visible when displaying the interchange import dialog with basic layout."))
 						.OnCheckStateChanged_Lambda([InterchangePipelinePtr = InterchangePipeline, PropertyPath](ECheckBoxState CheckType)
 						{
 							if (!ensure(InterchangePipelinePtr.IsValid()))
 							{
 								return;
 							}
-							FScopedTransaction ScopedTransaction(NSLOCTEXT("InterchangePipelineBaseDetails::CustomizeDetails", "TransactionvisibilityPropertiesToggleImport", "Toggle property visibility at import"), !GIsTransacting);
+							FScopedTransaction ScopedTransaction(LOCTEXT("TransactionvisibilityPropertiesToggleImport", "Toggle property visibility at import"), !GIsTransacting);
 							InterchangePipelinePtr->Modify();
 							InterchangePipelinePtr->FindOrAddPropertyStates(PropertyPath).SetPropertyBasicLayoutVisibility((CheckType != ECheckBoxState::Checked));
 							InterchangePipelinePtr->PostEditChange();
@@ -548,7 +632,7 @@ void FInterchangePipelineBaseDetailsCustomization::CustomizeDetails(IDetailLayou
 						SNew(STextBlock)
 						.Visibility(PipelineInternalEditionData ? EVisibility::Collapsed : EVisibility::All)
 						.Font(IDetailLayoutBuilder::GetDetailFont())
-						.Text(NSLOCTEXT("InterchangePipelineBaseDetails::CustomizeDetails", "HiddenAtImportText", "Import"))
+						.Text(LOCTEXT("HiddenAtImportText", "Import"))
 					]
 					+ SHorizontalBox::Slot()
 					.AutoWidth()
@@ -562,14 +646,14 @@ void FInterchangePipelineBaseDetailsCustomization::CustomizeDetails(IDetailLayou
 						.UncheckedImage(FAppStyle::Get().GetBrush("Icons.Visible"))
 						.UncheckedHoveredImage(FAppStyle::Get().GetBrush("Icons.Visible"))
 						.UncheckedPressedImage(FAppStyle::Get().GetBrush("Icons.Visible"))
-						.ToolTipText(NSLOCTEXT("InterchangePipelineBaseDetails::CustomizeDetails", "VisibleTooltipImport", "If true this property will be visible when displaying the interchange import dialog."))
+						.ToolTipText(LOCTEXT("VisibleTooltipImport", "If true this property will be visible when displaying the interchange import dialog."))
 						.OnCheckStateChanged_Lambda([InterchangePipelinePtr = InterchangePipeline, PropertyPath](ECheckBoxState CheckType)
 						{
 							if (!ensure(InterchangePipelinePtr.IsValid()))
 							{
 								return;
 							}
-							FScopedTransaction ScopedTransaction(NSLOCTEXT("InterchangePipelineBaseDetails::CustomizeDetails", "TransactionvisibilityPropertiesToggleImport", "Toggle property visibility at import"), !GIsTransacting);
+							FScopedTransaction ScopedTransaction(LOCTEXT("TransactionvisibilityPropertiesToggleImport", "Toggle property visibility at import"), !GIsTransacting);
 							InterchangePipelinePtr->Modify();
 							InterchangePipelinePtr->FindOrAddPropertyStates(PropertyPath).SetPropertyImportVisibility((CheckType != ECheckBoxState::Checked));
 							InterchangePipelinePtr->PostEditChange();
@@ -597,7 +681,7 @@ void FInterchangePipelineBaseDetailsCustomization::CustomizeDetails(IDetailLayou
 						SNew(STextBlock)
 						.Visibility(PipelineInternalEditionData ? EVisibility::Collapsed : EVisibility::All)
 						.Font(IDetailLayoutBuilder::GetDetailFont())
-						.Text(NSLOCTEXT("InterchangePipelineBaseDetails::CustomizeDetails", "HiddenAtReimportText", "Reimport"))
+						.Text(LOCTEXT("HiddenAtReimportText", "Reimport"))
 					]
 					+ SHorizontalBox::Slot()
 					.AutoWidth()
@@ -611,14 +695,14 @@ void FInterchangePipelineBaseDetailsCustomization::CustomizeDetails(IDetailLayou
 						.UncheckedImage(FAppStyle::Get().GetBrush("Icons.Visible"))
 						.UncheckedHoveredImage(FAppStyle::Get().GetBrush("Icons.Visible"))
 						.UncheckedPressedImage(FAppStyle::Get().GetBrush("Icons.Visible"))
-						.ToolTipText(NSLOCTEXT("InterchangePipelineBaseDetails::CustomizeDetails", "VisibleTooltipReimport", "If true this property will be visible when displaying the interchange reimport dialog."))
+						.ToolTipText(LOCTEXT("VisibleTooltipReimport", "If true this property will be visible when displaying the interchange reimport dialog."))
 						.OnCheckStateChanged_Lambda([InterchangePipelinePtr = InterchangePipeline, PropertyPath](ECheckBoxState CheckType)
 						{
 							if (!ensure(InterchangePipelinePtr.IsValid()))
 							{
 								return;
 							}
-							FScopedTransaction ScopedTransaction(NSLOCTEXT("InterchangePipelineBaseDetails::CustomizeDetails", "TransactionvisibilityPropertiesToggleReimport", "Toggle property visibility at reimport"), !GIsTransacting);
+							FScopedTransaction ScopedTransaction(LOCTEXT("TransactionvisibilityPropertiesToggleReimport", "Toggle property visibility at reimport"), !GIsTransacting);
 							InterchangePipelinePtr->Modify();
 							InterchangePipelinePtr->FindOrAddPropertyStates(PropertyPath).SetPropertyReimportVisibility((CheckType != ECheckBoxState::Checked));
 							InterchangePipelinePtr->PostEditChange();
diff --git a/Engine/Plugins/Interchange/Editor/Source/Pipelines/Private/InterchangeEditorPipelineDetails.h b/Engine/Plugins/Interchange/Editor/Source/Pipelines/Private/InterchangeEditorPipelineDetails.h
--- a/Engine/Plugins/Interchange/Editor/Source/Pipelines/Private/InterchangeEditorPipelineDetails.h
+++ b/Engine/Plugins/Interchange/Editor/Source/Pipelines/Private/InterchangeEditorPipelineDetails.h
@@ -3,6 +3,7 @@
 #pragma once
 
 #include "IDetailCustomization.h"
+#include "Input/Reply.h"
 #include "UObject/WeakObjectPtr.h"
 #include "Templates/SharedPointer.h"
 #include "Types/SlateEnums.h"
@@ -15,6 +16,7 @@ class IDetailPropertyRow;
 class IPropertyHandle;
 class UInterchangeBaseNode;
 class UInterchangePipelineBase;
+struct FInterchangeConflictInfo;
 
 namespace UE
 {
@@ -31,6 +33,9 @@ public:
 	/** Makes a new instance of this detail layout class for a specific detail view requesting it */
 	static TSharedRef<IDetailCustomization> MakeInstance();
 
+	// Temporary: Used to pass data to the next instance created.
+	static void SetConflictsInfo(TArray<FInterchangeConflictInfo>& ConflictInfos);
+
 	/** IDetailCustomization interface */
 	virtual void CustomizeDetails(IDetailLayoutBuilder& DetailBuilder) override;
 	/** End IDetailCustomization interface */
@@ -45,7 +50,9 @@ private:
 	void LockPropertyHandleRow(const TSharedPtr<IPropertyHandle> PropertyHandle, IDetailPropertyRow& PropertyRow) const;
 	void AddSubCategory(IDetailLayoutBuilder& DetailBuilder, TMap<FName, TMap<FName, TArray<FInternalPropertyData>>>& SubCategoriesPropertiesPerMainCategory);
 	void InternalGetPipelineProperties(const UInterchangePipelineBase* Pipeline, const TArray<FName>& AllCategoryNames, TMap<FName, TArray<FName>>& PropertiesPerCategorys) const;
-	
+	void AddConflictSection();
+	FReply ShowConflictDialog(FInterchangeConflictInfo ConflictInfo);
+
 	void SetTextComboBoxWidget(IDetailPropertyRow& PropertyRow, const TSharedPtr<IPropertyHandle>& Handle, const TArray<FString>& PossibleValues);
 	void OnTextComboBoxChanged(TSharedPtr<FString> NewValue, ESelectInfo::Type SelectInfo, TWeakPtr<IPropertyHandle> HandlePtr);
 	struct FInternalComboBoxData
diff --git a/Engine/Plugins/Interchange/Editor/Source/Pipelines/Private/InterchangeEditorPipelinesModule.cpp b/Engine/Plugins/Interchange/Editor/Source/Pipelines/Private/InterchangeEditorPipelinesModule.cpp
--- a/Engine/Plugins/Interchange/Editor/Source/Pipelines/Private/InterchangeEditorPipelinesModule.cpp
+++ b/Engine/Plugins/Interchange/Editor/Source/Pipelines/Private/InterchangeEditorPipelinesModule.cpp
@@ -167,6 +167,7 @@ void FInterchangeEditorPipelinesModule::RegisterPropertySectionMappings()
 	{
 		TSharedRef<FPropertySection> Section = RegisterPropertySection(PropertyModule, "InterchangeGenericAssetsPipeline", "General", LOCTEXT("General", "General"));
 		Section->AddCategory("Common");
+		Section->AddCategory("Conflicts");
 	}
 
 	// Materials
diff --git a/Engine/Plugins/Interchange/Editor/Source/Pipelines/Private/SInterchangePipelineConfigurationDialog.cpp b/Engine/Plugins/Interchange/Editor/Source/Pipelines/Private/SInterchangePipelineConfigurationDialog.cpp
--- a/Engine/Plugins/Interchange/Editor/Source/Pipelines/Private/SInterchangePipelineConfigurationDialog.cpp
+++ b/Engine/Plugins/Interchange/Editor/Source/Pipelines/Private/SInterchangePipelineConfigurationDialog.cpp
@@ -1,6 +1,8 @@
 // Copyright Epic Games, Inc. All Rights Reserved.
 #include "SInterchangePipelineConfigurationDialog.h"
 
+#include "InterchangeEditorPipelineDetails.h"
+
 #include "DetailsViewArgs.h"
 #include "Dialog/SCustomDialog.h"
 #include "Framework/MultiBox/MultiBoxBuilder.h"
@@ -52,8 +54,6 @@ void SInterchangePipelineItem::Construct(
 	TObjectPtr<UInterchangePipelineBase> PipelineElementPtr = PipelineElement->Pipeline;
 	check(PipelineElementPtr.Get());
 	FText PipelineName = LOCTEXT("InvalidPipelineName", "Invalid Pipeline");
-	FText ConflictsComboEntryText = LOCTEXT("SInterchangePipelineItem::Conflicts", "Conflicts");
-	ConflictsComboEntry = MakeShared<FString>(ConflictsComboEntryText.ToString());
 	if (PipelineElementPtr.Get())
 	{
 		FString PipelineNameString = PipelineElement->DisplayName;
@@ -62,23 +62,12 @@ void SInterchangePipelineItem::Construct(
 			PipelineNameString += FString::Printf(TEXT(" (%s)"), *PipelineElementPtr->GetClass()->GetName());
 		}
 		PipelineName = FText::FromString(PipelineNameString);
-		PipelineElement->ConflictInfos = PipelineElementPtr->GetConflictInfos(PipelineElement->ReimportObject, PipelineElement->Container, PipelineElement->SourceData);
-		if (PipelineElement->ConflictInfos.Num() > 0)
-		{
-			ConflictNameList.Reset(PipelineElement->ConflictInfos.Num() + 1);
-			ConflictNameList.Add(ConflictsComboEntry);
-			for (const FInterchangeConflictInfo& ConflictInfo : PipelineElement->ConflictInfos)
-			{
-				TSharedPtr<FString> ConflictNamePtr = MakeShared<FString>(ConflictInfo.DisplayName);
-				ConflictNameList.Add(ConflictNamePtr);
-			}
-		}
 	}
 	
-	ConflictComboBox = nullptr;
+	static const FSlateBrush* ConflictBrush = FAppStyle::GetBrush("Icons.Error");
+	const FText ConflictsComboBoxTooltip = LOCTEXT("ConflictsComboBoxTooltip", "If there is some conflict, simply select one to see more details.");
+	const FText Conflict_IconTooltip = FText::Format(LOCTEXT("Conflict_IconTooltip", "There are {0} conflicts. See Conflicts section below for details."), PipelineElement->ConflictInfos.Num());
 
-	FText ConflictsComboBoxTooltip = LOCTEXT("ConflictsComboBoxTooltip", "If there is some conflict, simply select one to see more details.");
-		
 	STableRow<TSharedPtr<FInterchangePipelineItemType>>::Construct(
 		STableRow<TSharedPtr<FInterchangePipelineItemType>>::FArguments()
 		.Content()
@@ -92,6 +81,19 @@ void SInterchangePipelineItem::Construct(
 				.Image(this, &SInterchangePipelineItem::GetImageItemIcon)
 			]
 			+ SHorizontalBox::Slot()
+			.AutoWidth()
+			.Padding(0.0f, 2.0f, 6.0f, 2.0f)
+			[
+				SNew(SImage)
+				.ToolTipText(Conflict_IconTooltip)
+				.Image(ConflictBrush)
+				.Visibility_Lambda([this]()->EVisibility
+					{
+						return PipelineElement->ConflictInfos.Num() > 0 ? EVisibility::All : EVisibility::Collapsed;
+					})
+				.ColorAndOpacity(this, &SInterchangePipelineItem::GetTextColor)
+			]
+			+ SHorizontalBox::Slot()
 			.FillWidth(1.0f)
 			.Padding(3.0f, 0.0f)
 			.VAlign(VAlign_Center)
@@ -100,48 +102,7 @@ void SInterchangePipelineItem::Construct(
 				.Text(PipelineName)
 				.ColorAndOpacity(this, &SInterchangePipelineItem::GetTextColor)
 			]
-			+ SHorizontalBox::Slot()
-			.AutoWidth()
-			.Padding(2.0f, 0.0f)
-			.VAlign(VAlign_Center)
-			[
-				SAssignNew(ConflictComboBox, STextComboBox)
-					.Visibility_Lambda([this]()->EVisibility
-						{
-							return PipelineElement->ConflictInfos.Num() > 0 ? EVisibility::All : EVisibility::Collapsed;
-						})
-					.OptionsSource(&ConflictNameList)
-					.OnSelectionChanged_Lambda([this](TSharedPtr<FString> String, ESelectInfo::Type)
-						{
-							if (!String.IsValid() || String->IsEmpty() || ConflictsComboEntry.Get()->Equals(*String.Get()))
-							{
-								return;
-							}
-							//Find and display the conflict info
-							for (const FInterchangeConflictInfo& ConflictInfo : PipelineElement->ConflictInfos)
-							{
-								if (ConflictInfo.DisplayName.Equals(*String.Get()))
-								{
-									if (ConflictInfo.Pipeline)
-									{
-										ConflictInfo.Pipeline->ShowConflictDialog(ConflictInfo.UniqueId);
-									}
-
-									//Re-select the conflict item after we show the conflict modal dialog
-									ConflictComboBox->SetSelectedItem(ConflictsComboEntry);
-									break;
-								}
-							}
-						})
-					.ToolTipText(ConflictsComboBoxTooltip)
-			]
 		], OwnerTable);
-
-	if (PipelineElement->ConflictInfos.Num() > 0)
-	{
-		//Select the conflicts item
-		ConflictComboBox->SetSelectedItem(ConflictsComboEntry);
-	}
 }
 
 const FSlateBrush* SInterchangePipelineItem::GetImageItemIcon() const
@@ -221,10 +182,20 @@ FString SInterchangePipelineConfigurationDialog::GetPipelineDisplayName(const UI
 	return PipelineDisplayName;
 }
 
-void SInterchangePipelineConfigurationDialog::SetEditPipeline(UInterchangePipelineBase* PipelineToEdit)
+void SInterchangePipelineConfigurationDialog::SetEditPipeline(FInterchangePipelineItemType* PipelineItemToEdit)
 {
 	TArray<UObject*> ObjectsToEdit;
-	ObjectsToEdit.Add(PipelineToEdit);
+	ObjectsToEdit.Add(!PipelineItemToEdit ? nullptr : PipelineItemToEdit->Pipeline);
+
+	if (PipelineItemToEdit)
+	{
+		PipelineItemToEdit->ConflictInfos.Reset();
+		if (PipelineItemToEdit->ReimportObject)
+		{
+			PipelineItemToEdit->ConflictInfos = PipelineItemToEdit->Pipeline->GetConflictInfos(PipelineItemToEdit->ReimportObject, PipelineItemToEdit->Container, PipelineItemToEdit->SourceData);
+		}
+		FInterchangePipelineBaseDetailsCustomization::SetConflictsInfo(PipelineItemToEdit->ConflictInfos);
+	}
 	PipelineConfigurationDetailsView->SetObjects(ObjectsToEdit);
 }
 
@@ -804,7 +775,9 @@ FReply SInterchangePipelineConfigurationDialog::OnResetToDefault()
 	{
 		return Result;
 	}
-	UInterchangePipelineBase* PipelineToEdit = nullptr;
+
+	FInterchangePipelineItemType* PipelineToEdit = nullptr;
+
 	//Multi selection is not allowed
 	for(TWeakObjectPtr<UObject> WeakObject : SelectedPipelines)
 	{
@@ -839,7 +812,7 @@ FReply SInterchangePipelineConfigurationDialog::OnResetToDefault()
 									}
 									//Switch the pipeline the element point on
 									PipelineListViewItems[PipelineIndex]->Pipeline = GeneratedPipeline;
-									PipelineToEdit = GeneratedPipeline;
+									PipelineToEdit = PipelineListViewItems[PipelineIndex].Get();
 									PipelinesListView->SetSelection(PipelineListViewItems[PipelineIndex], ESelectInfo::Direct);
 									PipelinesListView->RequestListRefresh();
 									break;
@@ -1077,7 +1050,7 @@ void SInterchangePipelineConfigurationDialog::OnPipelineSelectionChanged(TShared
 	{
 		CurrentSelectedPipeline = InItem->Pipeline;
 	}
-	SetEditPipeline(CurrentSelectedPipeline.Get());
+	SetEditPipeline(InItem.Get());
 	
 	if (CurrentSelectedPipeline)
 	{
diff --git a/Engine/Plugins/Interchange/Editor/Source/Pipelines/Private/SInterchangePipelineConfigurationDialog.h b/Engine/Plugins/Interchange/Editor/Source/Pipelines/Private/SInterchangePipelineConfigurationDialog.h
--- a/Engine/Plugins/Interchange/Editor/Source/Pipelines/Private/SInterchangePipelineConfigurationDialog.h
+++ b/Engine/Plugins/Interchange/Editor/Source/Pipelines/Private/SInterchangePipelineConfigurationDialog.h
@@ -42,9 +42,6 @@ private:
 	FSlateColor GetTextColor() const;
 
 	TSharedPtr<FInterchangePipelineItemType> PipelineElement = nullptr;
-	TArray<TSharedPtr<FString>> ConflictNameList;
-	TSharedPtr<FString> ConflictsComboEntry = nullptr;
-	TSharedPtr<STextComboBox> ConflictComboBox = nullptr;
 };
 
 typedef SListView< TSharedPtr<FInterchangePipelineItemType> > SPipelineListViewType;
@@ -116,7 +113,7 @@ private:
 	/** Internal utility function to properly display pipeline's name */
 	static FString GetPipelineDisplayName(const UInterchangePipelineBase* Pipeline);
 
-	void SetEditPipeline(UInterchangePipelineBase* PipelineToEdit);
+	void SetEditPipeline(FInterchangePipelineItemType* PipelineItemToEdit);
 	FReply OnEditTranslatorSettings();
 	void OnFinishedChangingProperties(const FPropertyChangedEvent& PropertyChangedEvent);
 
