Fixed a bAllowShrinking argument.

diff --git a/Engine/Plugins/Interchange/Editor/Source/Pipelines/Private/InterchangeEditorPipelineDetails.cpp b/Engine/Plugins/Interchange/Editor/Source/Pipelines/Private/InterchangeEditorPipelineDetails.cpp
--- a/Engine/Plugins/Interchange/Editor/Source/Pipelines/Private/InterchangeEditorPipelineDetails.cpp
+++ b/Engine/Plugins/Interchange/Editor/Source/Pipelines/Private/InterchangeEditorPipelineDetails.cpp
@@ -272,7 +272,7 @@ void FInterchangePipelineBaseDetailsCustomization::AddConflictSection()
 		return;
 	}
 
-	TArray<FInterchangeConflictInfo> ConflictInfos = ConflicInfosStack.Pop(false);
+	TArray<FInterchangeConflictInfo> ConflictInfos = ConflicInfosStack.Pop(EAllowShrinking::No);
 	if (!InterchangePipeline->IsReimportContext() || ConflictInfos.Num() == 0)
 	{
 		return;
