Interchange: Fixed an unfortunate typo in code which disabled a feature

#rnx

diff --git a/Engine/Plugins/Interchange/Editor/Source/Pipelines/Private/InterchangeEditorPipelineDetails.cpp b/Engine/Plugins/Interchange/Editor/Source/Pipelines/Private/InterchangeEditorPipelineDetails.cpp
--- a/Engine/Plugins/Interchange/Editor/Source/Pipelines/Private/InterchangeEditorPipelineDetails.cpp
+++ b/Engine/Plugins/Interchange/Editor/Source/Pipelines/Private/InterchangeEditorPipelineDetails.cpp
@@ -109,13 +109,13 @@ TSharedRef<IDetailCustomization> FInterchangePipelineBaseDetailsCustomization::M
 	return MakeShareable(new FInterchangePipelineBaseDetailsCustomization);
 }
 
-static TArray<TArray<FInterchangeConflictInfo>> ConflicInfosStack;
+static TArray<TArray<FInterchangeConflictInfo>> ConflictInfosStack;
 void FInterchangePipelineBaseDetailsCustomization::SetConflictsInfo(TArray<FInterchangeConflictInfo>& ConflictInfos)
 {
 	// Only add info when there are any
-	if (ConflicInfosStack.Num() > 0)
+	if (ConflictInfos.Num() > 0)
 	{
-		ConflicInfosStack.Push(ConflictInfos);
+		ConflictInfosStack.Push(ConflictInfos);
 	}
 }
 
@@ -270,12 +270,12 @@ FReply FInterchangePipelineBaseDetailsCustomization::ShowConflictDialog(FInterch
 
 void FInterchangePipelineBaseDetailsCustomization::AddConflictSection()
 {
-	if (ConflicInfosStack.Num() == 0)
+	if (ConflictInfosStack.Num() == 0)
 	{
 		return;
 	}
 
-	TArray<FInterchangeConflictInfo> ConflictInfos = ConflicInfosStack.Pop(EAllowShrinking::No);
+	TArray<FInterchangeConflictInfo> ConflictInfos = ConflictInfosStack.Pop(EAllowShrinking::No);
 	if (!InterchangePipeline->IsReimportContext() || ConflictInfos.Num() == 0)
 	{
 		return;
