Interchange: Re-order the category buttons

#rnx

diff --git a/Engine/Plugins/Interchange/Editor/Source/Pipelines/Private/InterchangeEditorPipelinesModule.cpp b/Engine/Plugins/Interchange/Editor/Source/Pipelines/Private/InterchangeEditorPipelinesModule.cpp
--- a/Engine/Plugins/Interchange/Editor/Source/Pipelines/Private/InterchangeEditorPipelinesModule.cpp
+++ b/Engine/Plugins/Interchange/Editor/Source/Pipelines/Private/InterchangeEditorPipelinesModule.cpp
@@ -163,6 +163,8 @@ void FInterchangeEditorPipelinesModule::RegisterPropertySectionMappings()
 	const FName PropertyEditorModuleName("PropertyEditor");
 	FPropertyEditorModule& PropertyModule = FModuleManager::GetModuleChecked<FPropertyEditorModule>(PropertyEditorModuleName);
 
+	//We add number in front of the section name because the property editor is ordering alphabetically. Since we set the display name, the section name in the UI are correct.
+
 	// Assets
 	{
 		TSharedRef<FPropertySection> Section = RegisterPropertySection(PropertyModule, "InterchangeGenericAssetsPipeline", "General", LOCTEXT("General", "General"));
@@ -170,37 +172,37 @@ void FInterchangeEditorPipelinesModule::RegisterPropertySectionMappings()
 		Section->AddCategory("Conflicts");
 	}
 
-	// Materials
+	// Static Meshes
 	{
-		TSharedRef<FPropertySection> Section = RegisterPropertySection(PropertyModule, "InterchangeGenericMaterialPipeline", "Materials", LOCTEXT("Materials", "Materials"));
-		Section->AddCategory("Materials");
+		TSharedRef<FPropertySection> Section = RegisterPropertySection(PropertyModule, "InterchangeGenericMeshPipeline", "1 StaticMeshes", LOCTEXT("Static Meshes", "Static Meshes"));
+		Section->AddCategory("Common Meshes");
+		Section->AddCategory("Static Meshes");
 	}
 
 	// Skeletal Meshes
 	{
-		TSharedRef<FPropertySection> Section = RegisterPropertySection(PropertyModule, "InterchangeGenericMeshPipeline", "SkeletalMeshes", LOCTEXT("Skeletal Meshes", "Skeletal Meshes"));
+		TSharedRef<FPropertySection> Section = RegisterPropertySection(PropertyModule, "InterchangeGenericMeshPipeline", "2 SkeletalMeshes", LOCTEXT("Skeletal Meshes", "Skeletal Meshes"));
 		Section->AddCategory("Common Meshes");
 		Section->AddCategory("Common Skeletal Meshes and Animations");
 		Section->AddCategory("Skeletal Meshes");
 	}
 
-	// Static Meshes
-	{
-		TSharedRef<FPropertySection> Section = RegisterPropertySection(PropertyModule, "InterchangeGenericMeshPipeline", "StaticMeshes", LOCTEXT("Static Meshes", "Static Meshes"));
-		Section->AddCategory("Common Meshes");
-		Section->AddCategory("Static Meshes");
-	}
-
 	// Animation
 	{
-		TSharedRef<FPropertySection> Section = RegisterPropertySection(PropertyModule, "InterchangeGenericAnimationPipeline", "AnimationSequences", LOCTEXT("Animation Sequences", "Animation Sequences"));
+		TSharedRef<FPropertySection> Section = RegisterPropertySection(PropertyModule, "InterchangeGenericAnimationPipeline", "3 AnimationSequences", LOCTEXT("Animations", "Animations"));
 		Section->AddCategory("Common Skeletal Meshes and Animations");
 		Section->AddCategory("Animations");
 	}
 
+	// Materials
+	{
+		TSharedRef<FPropertySection> Section = RegisterPropertySection(PropertyModule, "InterchangeGenericMaterialPipeline", "4 Materials", LOCTEXT("Materials", "Materials"));
+		Section->AddCategory("Materials");
+	}
+
 	// Textures
 	{
-		TSharedRef<FPropertySection> Section = RegisterPropertySection(PropertyModule, "InterchangeGenericTexturePipeline", "Textures", LOCTEXT("Textures", "Textures"));
+		TSharedRef<FPropertySection> Section = RegisterPropertySection(PropertyModule, "InterchangeGenericTexturePipeline", "5 Textures", LOCTEXT("Textures", "Textures"));
 		Section->AddCategory("Textures");
 	}
 }
