Interchange: Add the possibility to do editor blueprint pipeline

#rnx

diff --git a/Engine/Plugins/Interchange/Editor/Source/Pipelines/Private/InterchangeEditorPipelinesModule.cpp b/Engine/Plugins/Interchange/Editor/Source/Pipelines/Private/InterchangeEditorPipelinesModule.cpp
--- a/Engine/Plugins/Interchange/Editor/Source/Pipelines/Private/InterchangeEditorPipelinesModule.cpp
+++ b/Engine/Plugins/Interchange/Editor/Source/Pipelines/Private/InterchangeEditorPipelinesModule.cpp
@@ -2,6 +2,7 @@
 
 #include "InterchangeEditorPipelinesModule.h"
 
+#include "AssetTypeCategories.h"
 #include "InterchangeEditorPipelineDetails.h"
 #include "InterchangeEditorPipelineStyle.h"
 #include "InterchangeManager.h"
@@ -21,6 +22,12 @@
 
 class FInterchangeEditorPipelinesModule : public IInterchangeEditorPipelinesModule
 {
+public:
+	FInterchangeEditorPipelinesModule()
+		: InterchangeAssetCategory(EAssetTypeCategories::Misc)
+	{
+	}
+private:
 	/** IModuleInterface implementation */
 	virtual void StartupModule() override;
 	virtual void ShutdownModule() override;
@@ -35,7 +42,8 @@ class FInterchangeEditorPipelinesModule : public IInterchangeEditorPipelinesModu
 	void RegisterPropertySectionMappings();
 	void UnregisterPropertySectionMappings();
 
-private:
+	EAssetTypeCategories::Type InterchangeAssetCategory;
+
 	/** Pointer to the style set to use for the UI. */
 	TSharedPtr<ISlateStyle> InterchangeEditorPipelineStyle = nullptr;
 
@@ -101,14 +109,19 @@ void FInterchangeEditorPipelinesModule::AcquireResources()
 	// Register the InterchangeImportTestPlan asset
 	FAssetToolsModule& AssetToolsModule = FAssetToolsModule::GetModule();
 	IAssetTools& AssetTools = AssetToolsModule.Get();
+
+	InterchangeAssetCategory = AssetTools.RegisterAdvancedAssetCategory(FName(TEXT("Interchange")), NSLOCTEXT("InterchangeEditorPipelineModule", "InterchangeAssetCategoryCategoryName", "Interchange"));
 	
-	BlueprintPipelineBase_TypeActions = MakeShared<FAssetTypeActions_InterchangeBlueprintPipelineBase>();
+	PipelineBase_TypeActions = MakeShared<FAssetTypeActions_InterchangePipelineBase>(InterchangeAssetCategory);
+	AssetTools.RegisterAssetTypeActions(PipelineBase_TypeActions.ToSharedRef());
+
+	BlueprintPipelineBase_TypeActions = MakeShared<FAssetTypeActions_InterchangeBlueprintPipelineBase>(InterchangeAssetCategory);
 	AssetTools.RegisterAssetTypeActions(BlueprintPipelineBase_TypeActions.ToSharedRef());
 
-	PipelineBase_TypeActions = MakeShared<FAssetTypeActions_InterchangePipelineBase>();
-	AssetTools.RegisterAssetTypeActions(PipelineBase_TypeActions.ToSharedRef());
+	BlueprintEditorPipelineBase_TypeActions = MakeShared<FAssetTypeActions_InterchangeEditorBlueprintPipelineBase>(InterchangeAssetCategory);
+	AssetTools.RegisterAssetTypeActions(BlueprintEditorPipelineBase_TypeActions.ToSharedRef());
 
-	PythonPipelineBase_TypeActions = MakeShared<FAssetTypeActions_InterchangePythonPipelineBase>();
+	PythonPipelineBase_TypeActions = MakeShared<FAssetTypeActions_InterchangePythonPipelineBase>(InterchangeAssetCategory);
 	AssetTools.RegisterAssetTypeActions(PythonPipelineBase_TypeActions.ToSharedRef());
 
 	FCoreDelegates::OnPreExit.AddLambda([this]()
@@ -145,6 +158,7 @@ void FInterchangeEditorPipelinesModule::ReleaseResources()
 	{
 		IAssetTools& AssetTools = AssetToolsModule->Get();
 		AssetTools.UnregisterAssetTypeActions(BlueprintPipelineBase_TypeActions.ToSharedRef());
+		AssetTools.UnregisterAssetTypeActions(BlueprintEditorPipelineBase_TypeActions.ToSharedRef());
 		AssetTools.UnregisterAssetTypeActions(PipelineBase_TypeActions.ToSharedRef());
 		AssetTools.UnregisterAssetTypeActions(PythonPipelineBase_TypeActions.ToSharedRef());
 	}
diff --git a/Engine/Plugins/Interchange/Editor/Source/Pipelines/Private/InterchangePipelineFactories.cpp b/Engine/Plugins/Interchange/Editor/Source/Pipelines/Private/InterchangePipelineFactories.cpp
--- a/Engine/Plugins/Interchange/Editor/Source/Pipelines/Private/InterchangePipelineFactories.cpp
+++ b/Engine/Plugins/Interchange/Editor/Source/Pipelines/Private/InterchangePipelineFactories.cpp
@@ -3,12 +3,21 @@
 #include "InterchangePipelineFactories.h"
 
 #include "Engine/BlueprintGeneratedClass.h"
+#include "InterchangeEditorBlueprintPipelineBase.h"
 #include "InterchangeBlueprintPipelineBase.h"
 #include "InterchangePythonPipelineBase.h"
 #include "Kismet2/KismetEditorUtilities.h"
 #include "Kismet2/SClassPickerDialog.h"
 #include "Misc/MessageDialog.h"
 
+namespace UE::Interchange::Private
+{
+	FText GetInterchangeCategoryPath()
+	{
+		return NSLOCTEXT("InterchangeEditorPipeline", "GetInterchangeCategoryPath", "Interchange");
+	}
+}
+
 /**
  * UInterchangeBlueprintPipelineBaseFactory implementation
  */
@@ -62,6 +71,59 @@ UClass* FAssetTypeActions_InterchangeBlueprintPipelineBase::GetSupportedClass()
 	return UInterchangeBlueprintPipelineBase::StaticClass();
 }
 
+/**
+ * UInterchangeEditorBlueprintPipelineBaseFactory implementation
+ */
+
+UInterchangeEditorBlueprintPipelineBaseFactory::UInterchangeEditorBlueprintPipelineBaseFactory()
+{
+	bCreateNew = true;
+	bEditAfterNew = false;
+	SupportedClass = UInterchangeEditorBlueprintPipelineBase::StaticClass();
+	ParentClass = UInterchangeEditorPipelineBase::StaticClass();
+}
+
+UObject* UInterchangeEditorBlueprintPipelineBaseFactory::FactoryCreateNew(UClass* Class, UObject* InParent, FName Name, EObjectFlags Flags, UObject* Context, FFeedbackContext* Warn, FName CallingContext)
+{
+	// Make sure we are trying to factory a gameplay ability blueprint, then create and init one
+	check(Class->IsChildOf(UInterchangeEditorBlueprintPipelineBase::StaticClass()));
+
+	// If they selected an interface, force the parent class to be UInterface
+	if (BlueprintType == BPTYPE_Interface)
+	{
+		ParentClass = UInterface::StaticClass();
+	}
+
+	if ((ParentClass == nullptr) || !FKismetEditorUtilities::CanCreateBlueprintOfClass(ParentClass) || !ParentClass->IsChildOf(UInterchangeEditorPipelineBase::StaticClass()))
+	{
+		FFormatNamedArguments Args;
+		Args.Add(TEXT("ClassName"), (ParentClass != nullptr) ? FText::FromString(ParentClass->GetName()) : NSLOCTEXT("UInterchangeEditorBlueprintPipelineBaseFactory", "Null", "(null)"));
+		FMessageDialog::Open(EAppMsgType::Ok, FText::Format(NSLOCTEXT("UInterchangeEditorBlueprintPipelineBaseFactory", "CannotCreateInterchangeEditorBlueprintPipelineBase", "Cannot create an Interchange Editor Blueprint Pipeline based on the class '{ClassName}'."), Args));
+		return nullptr;
+	}
+	else
+	{
+		return CastChecked<UInterchangeEditorBlueprintPipelineBase>(FKismetEditorUtilities::CreateBlueprint(ParentClass, InParent, Name, BlueprintType, UInterchangeEditorBlueprintPipelineBase::StaticClass(), UBlueprintGeneratedClass::StaticClass(), CallingContext));
+	}
+}
+
+UObject* UInterchangeEditorBlueprintPipelineBaseFactory::FactoryCreateNew(UClass* Class, UObject* InParent, FName Name, EObjectFlags Flags, UObject* Context, FFeedbackContext* Warn)
+{
+	return FactoryCreateNew(Class, InParent, Name, Flags, Context, Warn, NAME_None);
+}
+
+UFactory* FAssetTypeActions_InterchangeEditorBlueprintPipelineBase::GetFactoryForBlueprintType(UBlueprint* InBlueprint) const
+{
+	UInterchangeEditorBlueprintPipelineBaseFactory* InterchangeEditorBlueprintPipelineBaseFactory = NewObject<UInterchangeEditorBlueprintPipelineBaseFactory>();
+	InterchangeEditorBlueprintPipelineBaseFactory->ParentClass = TSubclassOf<UInterchangeEditorPipelineBase>(*InBlueprint->GeneratedClass);
+	return InterchangeEditorBlueprintPipelineBaseFactory;
+}
+
+UClass* FAssetTypeActions_InterchangeEditorBlueprintPipelineBase::GetSupportedClass() const
+{
+	return UInterchangeEditorBlueprintPipelineBase::StaticClass();
+}
+
 /**
  * UInterchangePipelineBaseFactory implementation
  */
@@ -84,11 +146,6 @@ UObject* UInterchangePipelineBaseFactory::FactoryCreateNew(UClass* Class, UObjec
 	return Pipeline;
 }
 
-uint32 UInterchangePipelineBaseFactory::GetMenuCategories() const
-{
-	return EAssetTypeCategories::Misc;
-}
-
 FText UInterchangePipelineBaseFactory::GetDisplayName() const
 {
 	return NSLOCTEXT("UInterchangePipelineBaseFactory", "MenuEntry", "Interchange Pipeline");
@@ -108,6 +165,7 @@ bool UInterchangePipelineBaseFactory::ConfigureProperties()
 	Filter->AllowedChildrenOfClasses.Add(UInterchangePipelineBase::StaticClass());
 	//Blueprint pipeline have there own factory
 	Filter->DisallowedChildrenOfClasses.Add(UInterchangeBlueprintPipelineBase::StaticClass());
+	Filter->DisallowedChildrenOfClasses.Add(UInterchangeEditorBlueprintPipelineBase::StaticClass());
 
 	const FText TitleText = NSLOCTEXT("UInterchangePipelineBaseFactory", "CreateOptions", "Pick a Pipeline Class");
 	UClass* ChosenClass = nullptr;
@@ -163,11 +221,6 @@ UObject* UInterchangePythonPipelineAssetFactory::FactoryCreateNew(UClass* Class,
 	return Pipeline;
 }
 
-uint32 UInterchangePythonPipelineAssetFactory::GetMenuCategories() const
-{
-	return EAssetTypeCategories::Misc;
-}
-
 FText UInterchangePythonPipelineAssetFactory::GetDisplayName() const
 {
 	return NSLOCTEXT("UInterchangePythonPipelineFactory", "MenuEntry", "Interchange Python Pipeline");
@@ -190,6 +243,7 @@ bool UInterchangePythonPipelineAssetFactory::ConfigureProperties()
 
 	//Blueprint pipeline have there own factory
 	Filter->DisallowedChildrenOfClasses.Add(UInterchangeBlueprintPipelineBase::StaticClass());
+	Filter->DisallowedChildrenOfClasses.Add(UInterchangeEditorBlueprintPipelineBase::StaticClass());
 
 	const FText TitleText = NSLOCTEXT("UInterchangePythonPipelineAssetFactory", "CreateOptions", "Pick a Pipeline Class");
 	UClass* ChosenClass = nullptr;
diff --git a/Engine/Plugins/Interchange/Editor/Source/Pipelines/Private/InterchangePipelineFactories.h b/Engine/Plugins/Interchange/Editor/Source/Pipelines/Private/InterchangePipelineFactories.h
--- a/Engine/Plugins/Interchange/Editor/Source/Pipelines/Private/InterchangePipelineFactories.h
+++ b/Engine/Plugins/Interchange/Editor/Source/Pipelines/Private/InterchangePipelineFactories.h
@@ -62,14 +62,62 @@ public:
 class FAssetTypeActions_InterchangeBlueprintPipelineBase : public FAssetTypeActions_Blueprint
 {
 public:
+	FAssetTypeActions_InterchangeBlueprintPipelineBase(EAssetTypeCategories::Type InAssetCategory)
+		: AssetCategory(InAssetCategory)
+	{ }
 	// FAssetTypeActions_Blueprint interface
 	virtual UFactory* GetFactoryForBlueprintType(UBlueprint* InBlueprint) const override;
 
 	// IAssetTypeActions Implementation
-	virtual FText GetName() const override { return NSLOCTEXT("FAssetTypeActions_InterchangeGraphInspectorPipeline", "InterchangeGraphInspectorPipelineName", "Interchange Blueprint Pipeline Base"); }
+	virtual FText GetName() const override { return NSLOCTEXT("FAssetTypeActions_InterchangeBlueprintPipelineBase", "InterchangeBlueprintPipelineBaseName", "Interchange Blueprint Pipeline"); }
 	virtual FColor GetTypeColor() const override { return FColor(10, 25, 175); }
 	virtual UClass* GetSupportedClass() const override;
-	virtual uint32 GetCategories() override { return EAssetTypeCategories::Blueprint | EAssetTypeCategories::Misc; }
+	virtual uint32 GetCategories() override { return AssetCategory; }
+
+private:
+	EAssetTypeCategories::Type AssetCategory;
+};
+
+UCLASS(hidecategories = Object, collapsecategories)
+class UInterchangeEditorBlueprintPipelineBaseFactory : public UFactory
+{
+	GENERATED_BODY()
+
+public:
+	UInterchangeEditorBlueprintPipelineBaseFactory();
+
+	// The type of blueprint that will be created
+	UPROPERTY(EditAnywhere, Category = InterchangeEditorBlueprintPipelineBaseFactory)
+	TEnumAsByte<enum EBlueprintType> BlueprintType;
+
+	// The parent class of the created blueprint
+	UPROPERTY(EditAnywhere, Category = InterchangeEditorBlueprintPipelineBaseFactory)
+	TSubclassOf<class UInterchangeEditorPipelineBase> ParentClass;
+
+	//Begin UFactory Interface
+	virtual UObject* FactoryCreateNew(UClass* Class, UObject* InParent, FName Name, EObjectFlags Flags, UObject* Context, FFeedbackContext* Warn, FName CallingContext) override;
+	virtual UObject* FactoryCreateNew(UClass* Class, UObject* InParent, FName Name, EObjectFlags Flags, UObject* Context, FFeedbackContext* Warn) override;
+	//End UFactory Interface	
+};
+
+class FAssetTypeActions_InterchangeEditorBlueprintPipelineBase : public FAssetTypeActions_Blueprint
+{
+public:
+	FAssetTypeActions_InterchangeEditorBlueprintPipelineBase(EAssetTypeCategories::Type InAssetCategory)
+		: AssetCategory(InAssetCategory)
+	{ }
+
+	// FAssetTypeActions_Blueprint interface
+	virtual UFactory* GetFactoryForBlueprintType(UBlueprint* InBlueprint) const override;
+
+	// IAssetTypeActions Implementation
+	virtual FText GetName() const override { return NSLOCTEXT("FAssetTypeActions_InterchangeEditorBlueprintPipelineBase", "InterchangeEditorBlueprintPipelineBaseName", "Interchange Editor Blueprint Pipeline"); }
+	virtual FColor GetTypeColor() const override { return FColor(10, 25, 175); }
+	virtual UClass* GetSupportedClass() const override;
+	virtual uint32 GetCategories() override { return AssetCategory; }
+
+private:
+	EAssetTypeCategories::Type AssetCategory;
 };
 
 UCLASS(hidecategories = Object, collapsecategories)
@@ -82,7 +130,6 @@ public:
 
 	// UFactory Interface
 	virtual UObject* FactoryCreateNew(UClass* Class, UObject* InParent, FName Name, EObjectFlags Flags, UObject* Context, FFeedbackContext* Warn) override;
-	virtual uint32 GetMenuCategories() const override;
 	virtual FText GetDisplayName() const override;
 	virtual bool ConfigureProperties() override;
 	// End of UFactory Interface
@@ -90,17 +137,21 @@ private:
 	UClass* PipelineClass = nullptr;
 };
 
-
-
 class FAssetTypeActions_InterchangePipelineBase : public FAssetTypeActions_Base
 {
 public:
-	// IAssetTypeActions Implementation
+	FAssetTypeActions_InterchangePipelineBase(EAssetTypeCategories::Type InAssetCategory)
+		: AssetCategory(InAssetCategory)
+	{ }
+
 	virtual FText GetName() const override { return NSLOCTEXT("FAssetTypeActions_InterchangePipelineBase", "InterchangePipelineBaseName", "Interchange Pipeline"); }
 	virtual FColor GetTypeColor() const override { return FColor(135, 200, 25); }
 	virtual UClass* GetSupportedClass() const override;
-	virtual uint32 GetCategories() override { return EAssetTypeCategories::Misc; }
+	virtual uint32 GetCategories() override { return AssetCategory; }
 	virtual void OpenAssetEditor(const TArray<UObject*>& InObjects, TSharedPtr<class IToolkitHost> EditWithinLevelEditor = TSharedPtr<IToolkitHost>()) override;
+
+private:
+	EAssetTypeCategories::Type AssetCategory;
 };
 
 UCLASS(hidecategories = Object, collapsecategories)
@@ -113,7 +164,6 @@ public:
 
 	// UFactory Interface
 	virtual UObject* FactoryCreateNew(UClass* Class, UObject* InParent, FName Name, EObjectFlags Flags, UObject* Context, FFeedbackContext* Warn) override;
-	virtual uint32 GetMenuCategories() const override;
 	virtual FText GetDisplayName() const override;
 	virtual bool ConfigureProperties() override;
 	// End of UFactory Interface
@@ -124,10 +174,17 @@ private:
 class FAssetTypeActions_InterchangePythonPipelineBase : public FAssetTypeActions_Base
 {
 public:
+	FAssetTypeActions_InterchangePythonPipelineBase(EAssetTypeCategories::Type InAssetCategory)
+		: AssetCategory(InAssetCategory)
+	{ }
+
 	// IAssetTypeActions Implementation
 	virtual FText GetName() const override { return NSLOCTEXT("InterchangePipelineFactories", "FAssetTypeActions_InterchangePythonPipeline", "Interchange Python Pipeline"); }
 	virtual FColor GetTypeColor() const override { return FColor(135, 200, 25); }
 	virtual UClass* GetSupportedClass() const override;
-	virtual uint32 GetCategories() override { return EAssetTypeCategories::Misc; }
+	virtual uint32 GetCategories() override { return AssetCategory; }
 	virtual void OpenAssetEditor(const TArray<UObject*>& InObjects, TSharedPtr<class IToolkitHost> EditWithinLevelEditor = TSharedPtr<IToolkitHost>()) override;
+
+private:
+	EAssetTypeCategories::Type AssetCategory;
 };
diff --git a/Engine/Plugins/Interchange/Editor/Source/Pipelines/Public/InterchangeEditorBlueprintPipelineBase.h b/Engine/Plugins/Interchange/Editor/Source/Pipelines/Public/InterchangeEditorBlueprintPipelineBase.h
new file mode 100644
--- /dev/null
+++ b/Engine/Plugins/Interchange/Editor/Source/Pipelines/Public/InterchangeEditorBlueprintPipelineBase.h
@@ -0,0 +1,62 @@
+// Copyright Epic Games, Inc. All Rights Reserved.
+
+#pragma once
+
+#include "Engine/Blueprint.h"
+#include "InterchangeBlueprintPipelineBase.h"
+#include "InterchangePipelineBase.h"
+#include "UObject/Object.h"
+#include "UObject/ObjectMacros.h"
+
+
+#include "InterchangeEditorBlueprintPipelineBase.generated.h"
+
+UCLASS(BlueprintType, MinimalAPI)
+class UInterchangeEditorPipelineBase : public UInterchangePipelineBase
+{
+	GENERATED_BODY()
+
+public:
+	//UObject Interface
+	virtual bool IsEditorOnly() const override
+	{
+		return true;
+	}
+	// End of UObject interface
+};
+
+/**
+ * This class allow users to create editor only Interchange blueprint pipeline.
+ */
+UCLASS(BlueprintType, MinimalAPI)
+class UInterchangeEditorBlueprintPipelineBase : public UInterchangeBlueprintPipelineBase
+{
+	GENERATED_BODY()
+public:
+
+	UInterchangeEditorBlueprintPipelineBase()
+	{
+		ParentClass = UInterchangeEditorPipelineBase::StaticClass();
+		//We must make sure the GeneratedClass is generated after the blueprint is loaded
+		bRecompileOnLoad = true;
+	}
+
+	//UObject Interface
+	virtual bool IsEditorOnly() const override
+	{
+		return true;
+	}
+	// End of UObject interface
+
+	// UBlueprint interface
+	virtual bool SupportedByDefaultBlueprintFactory() const override
+	{
+		return false;
+	}
+
+	virtual bool AlwaysCompileOnLoad() const override
+	{
+		return true;
+	}
+	// End of UBlueprint interface
+};
diff --git a/Engine/Plugins/Interchange/Editor/Source/Pipelines/Public/InterchangeEditorPipelinesModule.h b/Engine/Plugins/Interchange/Editor/Source/Pipelines/Public/InterchangeEditorPipelinesModule.h
--- a/Engine/Plugins/Interchange/Editor/Source/Pipelines/Public/InterchangeEditorPipelinesModule.h
+++ b/Engine/Plugins/Interchange/Editor/Source/Pipelines/Public/InterchangeEditorPipelinesModule.h
@@ -5,6 +5,7 @@
 #include "Modules/ModuleManager.h"
 
 class FAssetTypeActions_InterchangeBlueprintPipelineBase;
+class FAssetTypeActions_InterchangeEditorBlueprintPipelineBase;
 class FAssetTypeActions_InterchangePipelineBase;
 class FAssetTypeActions_InterchangePythonPipelineBase;
 
@@ -40,6 +41,7 @@ public:
 	TArray<FName> ClassesToUnregisterOnShutdown;
 
 	TSharedPtr<FAssetTypeActions_InterchangeBlueprintPipelineBase> BlueprintPipelineBase_TypeActions;
+	TSharedPtr<FAssetTypeActions_InterchangeEditorBlueprintPipelineBase> BlueprintEditorPipelineBase_TypeActions;
 	TSharedPtr<FAssetTypeActions_InterchangePipelineBase> PipelineBase_TypeActions;
 	TSharedPtr<FAssetTypeActions_InterchangePythonPipelineBase> PythonPipelineBase_TypeActions;
 };
diff --git a/Engine/Source/Editor/Blutility/Private/BlutilityModule.cpp b/Engine/Source/Editor/Blutility/Private/BlutilityModule.cpp
--- a/Engine/Source/Editor/Blutility/Private/BlutilityModule.cpp
+++ b/Engine/Source/Editor/Blutility/Private/BlutilityModule.cpp
@@ -217,8 +217,9 @@ public:
 		const UClass* BPClass = Blueprint ? Blueprint->GetClass() : nullptr;
 
 		if( BPClass && 
-			(BPClass->IsChildOf( UEditorUtilityBlueprint::StaticClass())
-			|| BPClass->IsChildOf(UEditorUtilityWidgetBlueprint::StaticClass())))
+			( BPClass->IsChildOf( UEditorUtilityBlueprint::StaticClass() )
+				|| BPClass->IsChildOf( UEditorUtilityWidgetBlueprint::StaticClass() )
+				|| IsEditorOnlyObject(BPClass) ) )
 		{
 			return true;
 		}
