Replace EVisibility::All by EVisibility::Visible when we want visible slate widget.

#rnx

diff --git a/Engine/Plugins/Interchange/Editor/Source/Pipelines/Private/InterchangeEditorPipelineDetails.cpp b/Engine/Plugins/Interchange/Editor/Source/Pipelines/Private/InterchangeEditorPipelineDetails.cpp
--- a/Engine/Plugins/Interchange/Editor/Source/Pipelines/Private/InterchangeEditorPipelineDetails.cpp
+++ b/Engine/Plugins/Interchange/Editor/Source/Pipelines/Private/InterchangeEditorPipelineDetails.cpp
@@ -574,7 +574,7 @@ void FInterchangePipelineBaseDetailsCustomization::CustomizeDetails(IDetailLayou
 					.Padding(3.0f, 1.0f)
 					[
 						SNew(SCheckBox)
-						.Visibility(PipelineInternalEditionData ? EVisibility::Collapsed : EVisibility::All)
+						.Visibility(PipelineInternalEditionData ? EVisibility::Collapsed : EVisibility::Visible)
 						.CheckedImage(FAppStyle::Get().GetBrush("Icons.Lock"))
 						.CheckedHoveredImage(FAppStyle::Get().GetBrush("Icons.Lock"))
 						.CheckedPressedImage(FAppStyle::Get().GetBrush("Icons.Lock"))
@@ -624,7 +624,7 @@ void FInterchangePipelineBaseDetailsCustomization::CustomizeDetails(IDetailLayou
 					.VAlign(VAlign_Center)
 					[
 						SNew(STextBlock)
-						.Visibility(PipelineInternalEditionData ? EVisibility::Collapsed : EVisibility::All)
+						.Visibility(PipelineInternalEditionData ? EVisibility::Collapsed : EVisibility::Visible)
 						.Font(IDetailLayoutBuilder::GetDetailFont())
 						.Text(LOCTEXT("ShowWhenBasicLayoutText", "Basic Layout"))
 					]
@@ -633,7 +633,7 @@ void FInterchangePipelineBaseDetailsCustomization::CustomizeDetails(IDetailLayou
 					.Padding(3.0f, 1.0f)
 					[
 						SNew(SCheckBox)
-						.Visibility(PipelineInternalEditionData ? EVisibility::Collapsed : EVisibility::All)
+						.Visibility(PipelineInternalEditionData ? EVisibility::Collapsed : EVisibility::Visible)
 						.CheckedImage(FAppStyle::Get().GetBrush("Icons.Hidden"))
 						.CheckedHoveredImage(FAppStyle::Get().GetBrush("Icons.Hidden"))
 						.CheckedPressedImage(FAppStyle::Get().GetBrush("Icons.Hidden"))
@@ -671,7 +671,7 @@ void FInterchangePipelineBaseDetailsCustomization::CustomizeDetails(IDetailLayou
 					.VAlign(VAlign_Center)
 					[
 						SNew(STextBlock)
-						.Visibility(PipelineInternalEditionData ? EVisibility::Collapsed : EVisibility::All)
+						.Visibility(PipelineInternalEditionData ? EVisibility::Collapsed : EVisibility::Visible)
 						.Font(IDetailLayoutBuilder::GetDetailFont())
 						.Text(LOCTEXT("HiddenAtImportText", "Import"))
 					]
@@ -680,7 +680,7 @@ void FInterchangePipelineBaseDetailsCustomization::CustomizeDetails(IDetailLayou
 					.Padding(3.0f, 1.0f)
 					[
 						SNew(SCheckBox)
-						.Visibility(PipelineInternalEditionData ? EVisibility::Collapsed : EVisibility::All)
+						.Visibility(PipelineInternalEditionData ? EVisibility::Collapsed : EVisibility::Visible)
 						.CheckedImage(FAppStyle::Get().GetBrush("Icons.Hidden"))
 						.CheckedHoveredImage(FAppStyle::Get().GetBrush("Icons.Hidden"))
 						.CheckedPressedImage(FAppStyle::Get().GetBrush("Icons.Hidden"))
@@ -720,7 +720,7 @@ void FInterchangePipelineBaseDetailsCustomization::CustomizeDetails(IDetailLayou
 					.VAlign(VAlign_Center)
 					[
 						SNew(STextBlock)
-						.Visibility(PipelineInternalEditionData ? EVisibility::Collapsed : EVisibility::All)
+						.Visibility(PipelineInternalEditionData ? EVisibility::Collapsed : EVisibility::Visible)
 						.Font(IDetailLayoutBuilder::GetDetailFont())
 						.Text(LOCTEXT("HiddenAtReimportText", "Reimport"))
 					]
@@ -729,7 +729,7 @@ void FInterchangePipelineBaseDetailsCustomization::CustomizeDetails(IDetailLayou
 					.Padding(3.0f, 1.0f)
 					[
 						SNew(SCheckBox)
-						.Visibility(PipelineInternalEditionData ? EVisibility::Collapsed : EVisibility::All)
+						.Visibility(PipelineInternalEditionData ? EVisibility::Collapsed : EVisibility::Visible)
 						.CheckedImage(FAppStyle::Get().GetBrush("Icons.Hidden"))
 						.CheckedHoveredImage(FAppStyle::Get().GetBrush("Icons.Hidden"))
 						.CheckedPressedImage(FAppStyle::Get().GetBrush("Icons.Hidden"))
diff --git a/Engine/Plugins/Interchange/Editor/Source/Pipelines/Private/SInterchangeGraphInspectorWindow.cpp b/Engine/Plugins/Interchange/Editor/Source/Pipelines/Private/SInterchangeGraphInspectorWindow.cpp
--- a/Engine/Plugins/Interchange/Editor/Source/Pipelines/Private/SInterchangeGraphInspectorWindow.cpp
+++ b/Engine/Plugins/Interchange/Editor/Source/Pipelines/Private/SInterchangeGraphInspectorWindow.cpp
@@ -430,7 +430,7 @@ TSharedRef<SBox> SInterchangeGraphInspectorWindow::SpawnGraphInspector()
 						SNew(SCheckBox)
 						.HAlign(HAlign_Center)
 						.OnCheckStateChanged(GraphInspectorTreeview.Get(), &SInterchangeGraphInspectorTreeView::OnToggleSelectAll)
-						.Visibility(bPreview ? EVisibility::Collapsed : EVisibility::All)
+						.Visibility(bPreview ? EVisibility::Collapsed : EVisibility::Visible)
 					]
 					+ SHorizontalBox::Slot()
 					.FillWidth(1.0f)
@@ -439,7 +439,7 @@ TSharedRef<SBox> SInterchangeGraphInspectorWindow::SpawnGraphInspector()
 					[
 						SNew(STextBlock)
 						.Text(LOCTEXT("GraphInspectorWindow_Scene_All", "All"))
-						.Visibility(bPreview ? EVisibility::Collapsed : EVisibility::All)
+						.Visibility(bPreview ? EVisibility::Collapsed : EVisibility::Visible)
 					]
 				]
 				+ SUniformGridPanel::Slot(1, 0)
diff --git a/Engine/Plugins/Interchange/Editor/Source/Pipelines/Private/SInterchangePipelineConfigurationDialog.cpp b/Engine/Plugins/Interchange/Editor/Source/Pipelines/Private/SInterchangePipelineConfigurationDialog.cpp
--- a/Engine/Plugins/Interchange/Editor/Source/Pipelines/Private/SInterchangePipelineConfigurationDialog.cpp
+++ b/Engine/Plugins/Interchange/Editor/Source/Pipelines/Private/SInterchangePipelineConfigurationDialog.cpp
@@ -89,7 +89,7 @@ void SInterchangePipelineItem::Construct(
 				.Image(ConflictBrush)
 				.Visibility_Lambda([this]()->EVisibility
 					{
-						return PipelineElement->ConflictInfos.Num() > 0 ? EVisibility::All : EVisibility::Collapsed;
+						return PipelineElement->ConflictInfos.Num() > 0 ? EVisibility::Visible : EVisibility::Collapsed;
 					})
 				.ColorAndOpacity(this, &SInterchangePipelineItem::GetTextColor)
 			]
@@ -372,7 +372,7 @@ TSharedRef<SBox> SInterchangePipelineConfigurationDialog::SpawnPipelineConfigura
 		StackTextComboBox = SNew(SHorizontalBox)
 		.Visibility_Lambda([this]()
 			{
-				return bBasicLayout ? EVisibility::Collapsed : EVisibility::All;
+				return bBasicLayout ? EVisibility::Collapsed : EVisibility::Visible;
 			})
 		+ SHorizontalBox::Slot()
 		.VAlign(VAlign_Center)
@@ -520,7 +520,7 @@ void SInterchangePipelineConfigurationDialog::Construct(const FArguments& InArgs
 					SNew(SButton)
 					.Visibility_Lambda([this]()
 						{
-							return !TranslatorSettings ? EVisibility::Collapsed : EVisibility::All;
+							return !TranslatorSettings ? EVisibility::Collapsed : EVisibility::Visible;
 						})
 					.ToolTipText(LOCTEXT("SInterchangePipelineConfigurationDialog_TranslatorSettings_Tooltip", "Edit translator project settings."))
 					.OnClicked(this, &SInterchangePipelineConfigurationDialog::OnEditTranslatorSettings)
@@ -571,7 +571,7 @@ void SInterchangePipelineConfigurationDialog::Construct(const FArguments& InArgs
 					.ToolTipText(LOCTEXT("SInterchangePipelineConfigurationDialog_FilterPipelineOptions_tooltip", "Filter the pipeline options using the source content data."))
 					.Visibility_Lambda([this]()
 						{
-							return bReimport ? EVisibility::Collapsed : EVisibility::All;
+							return bReimport ? EVisibility::Collapsed : EVisibility::Visible;
 						})
 					+ SHorizontalBox::Slot()
 					.AutoWidth()
diff --git a/Engine/Plugins/Interchange/Runtime/Source/Pipelines/Private/InterchangeGenericAssetsPipeline.cpp b/Engine/Plugins/Interchange/Runtime/Source/Pipelines/Private/InterchangeGenericAssetsPipeline.cpp
--- a/Engine/Plugins/Interchange/Runtime/Source/Pipelines/Private/InterchangeGenericAssetsPipeline.cpp
+++ b/Engine/Plugins/Interchange/Runtime/Source/Pipelines/Private/InterchangeGenericAssetsPipeline.cpp
@@ -1780,7 +1780,7 @@ FReply SInterchangeGenericAssetSkeletonConflictWidget::SetSectionVisible(EInterc
 
 EVisibility SInterchangeGenericAssetSkeletonConflictWidget::IsSectionVisible(EInterchangeSkeletonCompareSection SectionIndex)
 {
-	return bShowSectionFlag[SectionIndex] ? EVisibility::All : EVisibility::Collapsed;
+	return bShowSectionFlag[SectionIndex] ? EVisibility::Visible : EVisibility::Collapsed;
 }
 
 const FSlateBrush* SInterchangeGenericAssetSkeletonConflictWidget::GetCollapsableArrow(EInterchangeSkeletonCompareSection SectionIndex) const
diff --git a/Engine/Source/Editor/Persona/Private/PersonaMeshDetails.cpp b/Engine/Source/Editor/Persona/Private/PersonaMeshDetails.cpp
--- a/Engine/Source/Editor/Persona/Private/PersonaMeshDetails.cpp
+++ b/Engine/Source/Editor/Persona/Private/PersonaMeshDetails.cpp
@@ -5192,7 +5192,7 @@ TSharedRef<SWidget> FPersonaMeshDetails::OnGenerateCustomNameWidgetsForSection(i
 			.Padding(0, 2, 0, 0)
 			[
 				SNew(SBox)
-				.Visibility(LodIndex == 0 && !IsSectionChunked ? EVisibility::All : EVisibility::Collapsed)
+				.Visibility(LodIndex == 0 && !IsSectionChunked ? EVisibility::Visible : EVisibility::Collapsed)
 				[
 					SNew(SHorizontalBox)
 					+SHorizontalBox::Slot()
@@ -5395,12 +5395,12 @@ bool FPersonaMeshDetails::IsSectionEnabled(int32 LodIndex, int32 SectionIndex) c
 
 EVisibility FPersonaMeshDetails::ShowEnabledSectionDetail(int32 LodIndex, int32 SectionIndex) const
 {
-	return IsSectionEnabled(LodIndex, SectionIndex) ? EVisibility::All : EVisibility::Collapsed;
+	return IsSectionEnabled(LodIndex, SectionIndex) ? EVisibility::Visible : EVisibility::Collapsed;
 }
 
 EVisibility FPersonaMeshDetails::ShowDisabledSectionDetail(int32 LodIndex, int32 SectionIndex) const
 {
-	return IsSectionEnabled(LodIndex, SectionIndex) ? EVisibility::Collapsed : EVisibility::All;
+	return IsSectionEnabled(LodIndex, SectionIndex) ? EVisibility::Collapsed : EVisibility::Visible;
 }
 
 void FPersonaMeshDetails::OnSectionEnabledChanged(int32 LodIndex, int32 SectionIndex, bool bEnable)
@@ -5547,7 +5547,7 @@ EVisibility FPersonaMeshDetails::ShowSectionGenerateUpToSlider(int32 LodIndex, i
 	{
 		return EVisibility::Collapsed;
 	}
-	return SkeletalMeshPtr->GetImportedModel()->LODModels[LodIndex].Sections[SectionIndex].GenerateUpToLodIndex == -1 ? EVisibility::Collapsed : EVisibility::All;
+	return SkeletalMeshPtr->GetImportedModel()->LODModels[LodIndex].Sections[SectionIndex].GenerateUpToLodIndex == -1 ? EVisibility::Collapsed : EVisibility::Visible;
 }
 
 ECheckBoxState FPersonaMeshDetails::IsGenerateUpToSectionEnabled(int32 LodIndex, int32 SectionIndex) const
@@ -5658,7 +5658,7 @@ EVisibility FPersonaMeshDetails::LodComboBoxVisibilityForLodPicker() const
 	{
 		return EVisibility::Hidden;
 	}
-	return EVisibility::All;
+	return EVisibility::Visible;
 }
 
 bool FPersonaMeshDetails::IsLodComboBoxEnabledForLodPicker() const
diff --git a/Engine/Source/Editor/UnrealEd/Private/Fbx/FbxCompareWindow.cpp b/Engine/Source/Editor/UnrealEd/Private/Fbx/FbxCompareWindow.cpp
--- a/Engine/Source/Editor/UnrealEd/Private/Fbx/FbxCompareWindow.cpp
+++ b/Engine/Source/Editor/UnrealEd/Private/Fbx/FbxCompareWindow.cpp
@@ -139,7 +139,7 @@ FReply SFbxSkeltonConflictWindow::SetSectionVisible(EFBXCompareSection SectionIn
 
 EVisibility SFbxSkeltonConflictWindow::IsSectionVisible(EFBXCompareSection SectionIndex)
 {
-	return bShowSectionFlag[SectionIndex] ? EVisibility::All : EVisibility::Collapsed;
+	return bShowSectionFlag[SectionIndex] ? EVisibility::Visible : EVisibility::Collapsed;
 }
 
 const FSlateBrush* SFbxSkeltonConflictWindow::GetCollapsableArrow(EFBXCompareSection SectionIndex) const
diff --git a/Engine/Source/Editor/UnrealEd/Private/FbxMaterialConflictWindow.h b/Engine/Source/Editor/UnrealEd/Private/FbxMaterialConflictWindow.h
--- a/Engine/Source/Editor/UnrealEd/Private/FbxMaterialConflictWindow.h
+++ b/Engine/Source/Editor/UnrealEd/Private/FbxMaterialConflictWindow.h
@@ -90,7 +90,7 @@ public:
 
 	EVisibility CollapsePreviewVisibility()
 	{
-		return bIsPreviewConflict ? EVisibility::Collapsed : EVisibility::All;
+		return bIsPreviewConflict ? EVisibility::Collapsed : EVisibility::Visible;
 	}
 
 	virtual FReply OnKeyDown( const FGeometry& MyGeometry, const FKeyEvent& InKeyEvent ) override
