Implement GetWorld() for Interchange Editor Blueprint Pipelines so they are able to utilize Blueprint Function Libraries (which require a WorldContext provided by GetWorld())

diff --git a/Engine/Plugins/Interchange/Editor/Source/Pipelines/Private/InterchangeEditorBlueprintPipelineBase.cpp b/Engine/Plugins/Interchange/Editor/Source/Pipelines/Private/InterchangeEditorBlueprintPipelineBase.cpp
new file mode 100644
--- /dev/null
+++ b/Engine/Plugins/Interchange/Editor/Source/Pipelines/Private/InterchangeEditorBlueprintPipelineBase.cpp
@@ -0,0 +1,14 @@
+// Copyright Epic Games, Inc. All Rights Reserved.
+
+#include "InterchangeEditorBlueprintPipelineBase.h"
+#include "Editor.h"
+
+UWorld* UInterchangeEditorPipelineBase::GetWorld() const
+{
+	if (GEditor)
+	{
+		return GEditor->GetEditorWorldContext(false).World();
+	}
+
+	return Super::GetWorld();
+}
diff --git a/Engine/Plugins/Interchange/Editor/Source/Pipelines/Public/InterchangeEditorBlueprintPipelineBase.h b/Engine/Plugins/Interchange/Editor/Source/Pipelines/Public/InterchangeEditorBlueprintPipelineBase.h
--- a/Engine/Plugins/Interchange/Editor/Source/Pipelines/Public/InterchangeEditorBlueprintPipelineBase.h
+++ b/Engine/Plugins/Interchange/Editor/Source/Pipelines/Public/InterchangeEditorBlueprintPipelineBase.h
@@ -22,6 +22,8 @@ public:
 	{
 		return true;
 	}
+
+	virtual UWorld* GetWorld() const override;
 	// End of UObject interface
 };
 
