Interchange: Move Extra information section to the end of the UI and hide it if basic layout is enable.

#rnx

diff --git a/Engine/Plugins/Interchange/Editor/Source/Pipelines/Private/InterchangeEditorPipelineDetails.cpp b/Engine/Plugins/Interchange/Editor/Source/Pipelines/Private/InterchangeEditorPipelineDetails.cpp
--- a/Engine/Plugins/Interchange/Editor/Source/Pipelines/Private/InterchangeEditorPipelineDetails.cpp
+++ b/Engine/Plugins/Interchange/Editor/Source/Pipelines/Private/InterchangeEditorPipelineDetails.cpp
@@ -392,7 +392,7 @@ void FInterchangePipelineBaseDetailsCustomization::CustomizeDetails(IDetailLayou
 
 	const bool bAllowPropertyStatesEdition = InterchangePipeline->CanEditPropertiesStates();
 	const bool bIsReimportContext = InterchangePipeline->IsReimportContext();
-	const bool bISBasicLayout = InterchangePipeline->IsBasicLayout();
+	const bool bIsBasicLayout = InterchangePipeline->IsBasicLayout();
 
 	TArray<FName> AllCategoryNames;
 	CachedDetailBuilder->GetCategoryNames(AllCategoryNames);
@@ -401,8 +401,6 @@ void FInterchangePipelineBaseDetailsCustomization::CustomizeDetails(IDetailLayou
 	
 	AddConflictSection();
 	
-	AddExtraInformationSection();
-
 	for (const TPair<FName, TArray<FName>>& CategoryAndProperties : PropertiesPerCategorys)
 	{
 		//Category meta value Subgroup data
@@ -511,7 +509,7 @@ void FInterchangePipelineBaseDetailsCustomization::CustomizeDetails(IDetailLayou
 				bool IsLocked = false;
 				if (const FInterchangePipelinePropertyStates* PropertyStates = InterchangePipeline->GetPropertyStates(PropertyPath))
 				{
-					if (!PropertyStates->IsPropertyVisible(bIsReimportContext, bISBasicLayout))
+					if (!PropertyStates->IsPropertyVisible(bIsReimportContext, bIsBasicLayout))
 					{
 						continue;
 					}
@@ -784,6 +782,11 @@ void FInterchangePipelineBaseDetailsCustomization::CustomizeDetails(IDetailLayou
 			}
 		}
 	}
+
+	if (!bIsBasicLayout)
+	{
+		AddExtraInformationSection();
+	}
 }
 
 FInterchangeBaseNodeDetailsCustomization::FInterchangeBaseNodeDetailsCustomization()
