Interchange: Autosave disabled while Interchange is Running Import.

#rnx

diff --git a/Engine/Plugins/Interchange/Editor/Source/InterchangeEditor/Private/InterchangeEditorModule.cpp b/Engine/Plugins/Interchange/Editor/Source/InterchangeEditor/Private/InterchangeEditorModule.cpp
--- a/Engine/Plugins/Interchange/Editor/Source/InterchangeEditor/Private/InterchangeEditorModule.cpp
+++ b/Engine/Plugins/Interchange/Editor/Source/InterchangeEditor/Private/InterchangeEditorModule.cpp
@@ -12,6 +12,7 @@
 #include "Misc/App.h"
 #include "Misc/CoreDelegates.h"
 #include "Modules/ModuleManager.h"
+#include "Settings/EditorLoadingSavingSettings.h"
 
 #define LOCTEXT_NAMESPACE "InterchangeEditorModule"
 
@@ -19,6 +20,7 @@ DEFINE_LOG_CATEGORY(LogInterchangeEditor);
 
 namespace UE::Interchange::InterchangeEditorModule
 {
+	static bool bOldAutoSaveState = false;
 	bool HasErrorsOrWarnings(TStrongObjectPtr<UInterchangeResultsContainer> InResultsContainer)
 	{
 		for (UInterchangeResult* Result : InResultsContainer->GetResults())
@@ -65,6 +67,22 @@ namespace UE::Interchange::InterchangeEditorModule
 			LogListing->NotifyIfAnyMessages(NSLOCTEXT("Interchange", "LogAndNotify", "There were issues with the import."), EMessageSeverity::Info);
 		}
 	}
+
+	void ImportStarted()
+	{
+		//Store AutoSave setting and Set autosave to false:
+		UEditorLoadingSavingSettings* LoadingSavingSettings = GetMutableDefault<UEditorLoadingSavingSettings>();
+		// Disable autosaving while the Interchange is in progress.
+		bOldAutoSaveState = LoadingSavingSettings->bAutoSaveEnable;
+		LoadingSavingSettings->bAutoSaveEnable = false;
+	}
+
+	void ImportFinished()
+	{
+		//Reinstate AutoSave
+		UEditorLoadingSavingSettings* LoadingSavingSettings = GetMutableDefault<UEditorLoadingSavingSettings>();
+		LoadingSavingSettings->bAutoSaveEnable = bOldAutoSaveState;
+	}
 }
 
 FInterchangeEditorModule& FInterchangeEditorModule::Get()
@@ -84,15 +102,21 @@ void FInterchangeEditorModule::StartupModule()
 	auto RegisterItems = [this]()
 	{
 		FDelegateHandle InterchangeEditorModuleDelegate;
+		FDelegateHandle InterchangeEditorModuleDelegateOnImportStarted;
+		FDelegateHandle InterchangeEditorModuleDelegateOnImportFinished;
 
 		UInterchangeManager& InterchangeManager = UInterchangeManager::GetInterchangeManager();
 		InterchangeEditorModuleDelegate = InterchangeManager.OnBatchImportComplete.AddStatic(&InterchangeEditorModule::LogErrors);
+		InterchangeEditorModuleDelegateOnImportStarted = InterchangeManager.OnImportStarted.AddStatic(&InterchangeEditorModule::ImportStarted);
+		InterchangeEditorModuleDelegateOnImportFinished = InterchangeManager.OnImportFinished.AddStatic(&InterchangeEditorModule::ImportFinished);
 		InterchangeManager.RegisterImportDataConverter(UInterchangeFbxAssetImportDataConverter::StaticClass());
 
-		auto UnregisterItems = [InterchangeEditorModuleDelegate]()
+		auto UnregisterItems = [InterchangeEditorModuleDelegate, InterchangeEditorModuleDelegateOnImportStarted, InterchangeEditorModuleDelegateOnImportFinished]()
 		{
 			UInterchangeManager& InterchangeManager = UInterchangeManager::GetInterchangeManager();
 			InterchangeManager.OnBatchImportComplete.Remove(InterchangeEditorModuleDelegate);
+			InterchangeManager.OnImportStarted.Remove(InterchangeEditorModuleDelegateOnImportStarted);
+			InterchangeManager.OnImportFinished.Remove(InterchangeEditorModuleDelegateOnImportFinished);
 		};
 
 		InterchangeManager.OnPreDestroyInterchangeManager.AddLambda(UnregisterItems);
diff --git a/Engine/Source/Runtime/Interchange/Engine/Private/InterchangeManager.cpp b/Engine/Source/Runtime/Interchange/Engine/Private/InterchangeManager.cpp
--- a/Engine/Source/Runtime/Interchange/Engine/Private/InterchangeManager.cpp
+++ b/Engine/Source/Runtime/Interchange/Engine/Private/InterchangeManager.cpp
@@ -61,6 +61,8 @@ bool UInterchangeManager::bIsCreatingSingleton = false;
 
 namespace UE::Interchange::Private
 {
+	static TAtomic<uint64> AsyncHelperCounter = 0;
+
 	const FLogCategoryBase* GetLogInterchangePtr()
 	{
 #if NO_LOGGING
@@ -217,6 +219,23 @@ UE::Interchange::FImportAsyncHelper::FImportAsyncHelper()
 	, SceneImportResult(MakeShared<FImportResult>())
 {
 	bCancel = false;
+
+	if (UE::Interchange::Private::AsyncHelperCounter == 0)
+	{
+		UInterchangeManager::GetInterchangeManager().OnImportStarted.Broadcast();
+	}
+	UE::Interchange::Private::AsyncHelperCounter++;
+}
+
+UE::Interchange::FImportAsyncHelper::~FImportAsyncHelper()
+{
+	CleanUp();
+
+	UE::Interchange::Private::AsyncHelperCounter--;
+	if (UE::Interchange::Private::AsyncHelperCounter == 0)
+	{
+		UInterchangeManager::GetInterchangeManager().OnImportFinished.Broadcast();
+	}
 }
 
 void UE::Interchange::FImportAsyncHelper::AddReferencedObjects(FReferenceCollector& Collector)
@@ -2152,6 +2171,11 @@ void UInterchangeManager::SetInterchangeImportEnabled(bool bEnabled)
 	CCvarInterchangeImportEnable->Set(bEnabled);
 }
 
+bool UInterchangeManager::IsImporting()
+{
+	return UE::Interchange::Private::AsyncHelperCounter > 0;
+}
+
 bool UInterchangeManager::ExportAsset(const UObject* Asset, bool bIsAutomated)
 {
 	return false;
diff --git a/Engine/Source/Runtime/Interchange/Engine/Public/InterchangeManager.h b/Engine/Source/Runtime/Interchange/Engine/Public/InterchangeManager.h
--- a/Engine/Source/Runtime/Interchange/Engine/Public/InterchangeManager.h
+++ b/Engine/Source/Runtime/Interchange/Engine/Public/InterchangeManager.h
@@ -43,6 +43,11 @@ DECLARE_DELEGATE_OneParam(FOnObjectImportDoneNative, UObject*);
 DECLARE_DYNAMIC_DELEGATE_OneParam(FOnImportDoneDynamic, const TArray<UObject*>&, Objects);
 DECLARE_DELEGATE_OneParam(FOnImportDoneNative, const TArray<UObject*>&);
 
+/** Delegate Type that is fired when Interchange starts importing. Won't fire when a new import process starts while one is already in progress. Primarily Used for AutoSave setting. */
+DECLARE_MULTICAST_DELEGATE(FOnImportStarted);
+/** Delegate Type that is fired when Interchange finished importing. Won't fire when an import process finishes while one is still in progress. Primarily Used for AutoSave setting. */
+DECLARE_MULTICAST_DELEGATE(FOnImportFinished);
+
 namespace UE
 {
 	namespace Interchange
@@ -197,10 +202,7 @@ namespace UE
 		public:
 			FImportAsyncHelper();
 
-			~FImportAsyncHelper()
-			{
-				CleanUp();
-			}
+			~FImportAsyncHelper();
 
 			/* FGCObject interface */
 			virtual void AddReferencedObjects(FReferenceCollector& Collector) override;
@@ -432,6 +434,9 @@ public:
 	/** Set the CVar that enables or disables Interchange. */
 	static INTERCHANGEENGINE_API void SetInterchangeImportEnabled(bool bEnabled);
 
+	/** Checks if there are any imports in progress.*/
+	static INTERCHANGEENGINE_API bool IsImporting();
+
 	/** Delegate type that is fired when new assets have been imported. Note: InCreatedObject can be NULL if the import failed. Params: UFactory* InFactory, UObject* InCreatedObject. */
 	DECLARE_MULTICAST_DELEGATE_OneParam(FInterchangeOnAssetPostImport, UObject*);
 	/** Delegate type that is fired when new assets have been reimported. Note: InCreatedObject can be NULL if the import failed. Params: UObject* InCreatedObject. */
@@ -444,6 +449,10 @@ public:
 	FInterchangeOnAssetPostImport OnAssetPostImport;
 	FInterchangeOnAssetPostReimport OnAssetPostReimport;
 	FInterchangeOnBatchImportComplete OnBatchImportComplete;
+	//Fires when the first import process starts.
+	FOnImportStarted OnImportStarted;
+	//Fires when the last import process finishes.
+	FOnImportFinished OnImportFinished;
 	// Called when before the application is exiting.
 	FSimpleMulticastDelegate OnPreDestroyInterchangeManager;
 
