Interchange: Conversion should be able to use a valid pipeline asset instead of creating from a class.

#rnx

diff --git a/Engine/Config/BaseEngine.ini b/Engine/Config/BaseEngine.ini
--- a/Engine/Config/BaseEngine.ini
+++ b/Engine/Config/BaseEngine.ini
@@ -3515,6 +3515,7 @@ SceneImportSettings=(PipelineStacks=(("Scene", (Pipelines=("/Interchange/Pipelin
 FilePickerClass=/Script/InterchangeEditorUtilities.InterchangeFilePickerGeneric
 bStaticMeshUseSmoothEdgesIfSmoothingInformationIsMissing=True
 GenericPipelineClass=/Script/InterchangePipelines.InterchangeGenericAssetsPipeline
+ConverterDefaultPipeline=/Interchange/Pipelines/DefaultAssetsPipeline.DefaultAssetsPipeline
 
 [/Script/Engine.NetworkMetricsCSV_Replication]
 IntervalSeconds=0
diff --git a/Engine/Plugins/Interchange/Editor/Source/InterchangeEditor/Private/InterchangeFbxAssetImportDataConverter.cpp b/Engine/Plugins/Interchange/Editor/Source/InterchangeEditor/Private/InterchangeFbxAssetImportDataConverter.cpp
--- a/Engine/Plugins/Interchange/Editor/Source/InterchangeEditor/Private/InterchangeFbxAssetImportDataConverter.cpp
+++ b/Engine/Plugins/Interchange/Editor/Source/InterchangeEditor/Private/InterchangeFbxAssetImportDataConverter.cpp
@@ -24,6 +24,7 @@
 #include "InterchangeGenericMaterialPipeline.h"
 #include "InterchangeGenericMeshPipeline.h"
 #include "InterchangeGenericTexturePipeline.h"
+#include "InterchangeProjectSettings.h"
 #include "InterchangeSkeletalMeshFactoryNode.h"
 #include "InterchangeStaticMeshFactoryNode.h"
 #include "InterchangeManager.h"
@@ -31,6 +32,26 @@
 
 namespace UE::Interchange::Private
 {
+	//Create a generic asset pipeline, use the one from the project settings if its valid
+	UInterchangeGenericAssetsPipeline* GetDefaultGenericAssetPipelineForConvertion(UObject* Outer)
+	{
+		//Create a generic asset pipeline, use the one from the project settings if its valid
+		UInterchangeGenericAssetsPipeline* GenericAssetPipeline = nullptr;
+		if (const UInterchangeProjectSettings* InterchangeProjectSettings = GetDefault<UInterchangeProjectSettings>())
+		{
+			if (UInterchangeGenericAssetsPipeline* ConvertDefaultPipelineAsset = Cast<UInterchangeGenericAssetsPipeline>(InterchangeProjectSettings->ConverterDefaultPipeline.TryLoad()))
+			{
+				GenericAssetPipeline = DuplicateObject<UInterchangeGenericAssetsPipeline>(ConvertDefaultPipelineAsset, Outer);
+			}
+		}
+
+		if (!GenericAssetPipeline)
+		{
+			GenericAssetPipeline = NewObject<UInterchangeGenericAssetsPipeline>(Outer);
+		}
+		return GenericAssetPipeline;
+	}
+
 	void TransferSourceFileInformation(const UAssetImportData* SourceData, UAssetImportData* DestinationData)
 	{
 		TArray<FAssetImportInfo::FSourceFile> SourceFiles = SourceData->GetSourceData().SourceFiles;
@@ -549,8 +570,8 @@ namespace UE::Interchange::Private
 
 		if (!GenericAssetPipeline)
 		{
-			//Since we did not find any generic asset pipeline we fallback on the generic pipeline CDO settings
-			GenericAssetPipeline = NewObject<UInterchangeGenericAssetsPipeline>(GetTransientPackage());
+			//Since we did not find any generic asset pipeline we fallback on the generic pipeline from the project settings conversion
+			GenericAssetPipeline = GetDefaultGenericAssetPipelineForConvertion(GetTransientPackage());
 		}
 
 		FillFbxAssetImportData(InterchangeFbxTranslatorSettings, GenericAssetPipeline, DestinationData->StaticMeshImportData);
@@ -650,7 +671,7 @@ namespace UE::Interchange::Private
 		const FString NodeUniqueId = FGuid::NewGuid().ToString(EGuidFormats::Base36Encoded) + TEXT("_") + NodeDisplayLabel;
 
 		TArray<UObject*> Pipelines;
-		UInterchangeGenericAssetsPipeline* GenericAssetPipeline = NewObject<UInterchangeGenericAssetsPipeline>(DestinationData);
+		UInterchangeGenericAssetsPipeline* GenericAssetPipeline = GetDefaultGenericAssetPipelineForConvertion(DestinationData);
 		Pipelines.Add(GenericAssetPipeline);
 		DestinationData->SetPipelines(Pipelines);
 
@@ -724,7 +745,7 @@ namespace UE::Interchange::Private
 		DestinationData->SetNodeContainer(DestinationContainer);
 
 		TArray<UObject*> Pipelines;
-		UInterchangeGenericAssetsPipeline* GenericAssetPipeline = NewObject<UInterchangeGenericAssetsPipeline>(DestinationData);
+		UInterchangeGenericAssetsPipeline* GenericAssetPipeline = GetDefaultGenericAssetPipelineForConvertion(DestinationData);
 		Pipelines.Add(GenericAssetPipeline);
 		DestinationData->SetPipelines(Pipelines);
 
diff --git a/Engine/Source/Runtime/Interchange/Engine/Public/InterchangeProjectSettings.h b/Engine/Source/Runtime/Interchange/Engine/Public/InterchangeProjectSettings.h
--- a/Engine/Source/Runtime/Interchange/Engine/Public/InterchangeProjectSettings.h
+++ b/Engine/Source/Runtime/Interchange/Engine/Public/InterchangeProjectSettings.h
@@ -134,6 +134,13 @@ public:
 	 */
 	UPROPERTY(EditAnywhere, config, Category = "Editor Generic Pipeline Class")
 	TSoftClassPtr <UInterchangePipelineBase> GenericPipelineClass;
+
+	/**
+	 * Optional, the pipeline asset converters will duplicate to create interchange import data pipeline.
+	 * If not set, converters will duplicate a pipeline class CDO of there choice.
+	 */
+	UPROPERTY(EditAnywhere, config, Category = "Converters", meta = (AllowedClasses = "/Script/InterchangeCore.InterchangePipelineBase, /Script/InterchangeEngine.InterchangeBlueprintPipelineBase, /Script/InterchangeEngine.InterchangePythonPipelineAsset"))
+	FSoftObjectPath ConverterDefaultPipeline;
 };
 
 class FInterchangeProjectSettingsUtils
