Interchange: fix missing skeleton when converting legacy fbx to interchange skeletalmesh data.

#rnx

diff --git a/Engine/Plugins/Interchange/Editor/Source/InterchangeEditor/Private/InterchangeFbxAssetImportDataConverter.cpp b/Engine/Plugins/Interchange/Editor/Source/InterchangeEditor/Private/InterchangeFbxAssetImportDataConverter.cpp
--- a/Engine/Plugins/Interchange/Editor/Source/InterchangeEditor/Private/InterchangeFbxAssetImportDataConverter.cpp
+++ b/Engine/Plugins/Interchange/Editor/Source/InterchangeEditor/Private/InterchangeFbxAssetImportDataConverter.cpp
@@ -809,6 +809,11 @@ namespace UE::Interchange::Private
 			GenericAssetPipeline->MeshPipeline->bImportSkeletalMeshes = true;
 			GenericAssetPipeline->MeshPipeline->CommonMeshesProperties->ForceAllMeshAsType = EInterchangeForceMeshType::IFMT_SkeletalMesh;
 
+			if (FbxImportUI->Skeleton)
+			{
+				GenericAssetPipeline->CommonSkeletalMeshesAndAnimationsProperties->Skeleton = FbxImportUI->Skeleton;
+			}
+
 			GenericAssetPipeline->AnimationPipeline->bImportAnimations = FbxImportUI->bImportAnimations;
 
 			GenericAssetPipeline->ImportOffsetRotation = FbxImportUI->SkeletalMeshImportData->ImportRotation;
