Interchange: Refactor the UI/UX according to the propose design.

#rnx

diff --git a/Engine/Plugins/Interchange/Editor/Source/Pipelines/Private/InterchangeCardsPipeline.cpp b/Engine/Plugins/Interchange/Editor/Source/Pipelines/Private/InterchangeCardsPipeline.cpp
new file mode 100644
--- /dev/null
+++ b/Engine/Plugins/Interchange/Editor/Source/Pipelines/Private/InterchangeCardsPipeline.cpp
@@ -0,0 +1,24 @@
+// Copyright Epic Games, Inc. All Rights Reserved. 
+#include "InterchangeCardsPipeline.h"
+
+#include "CoreGlobals.h"
+#include "InterchangeSourceData.h"
+#include "Nodes/InterchangeBaseNodeContainer.h"
+
+#include UE_INLINE_GENERATED_CPP_BY_NAME(InterchangeCardsPipeline)
+
+void UInterchangeCardsPipeline::SetDisabledFactoryNodes(TArray<UClass*> FactoryNodeClasses)
+{
+	FactoryNodeClassesToDisabled = FactoryNodeClasses;
+}
+
+void UInterchangeCardsPipeline::ExecutePipeline(UInterchangeBaseNodeContainer* BaseNodeContainer, const TArray<UInterchangeSourceData*>& SourceDatas, const FString& ContentBasePath)
+{
+	BaseNodeContainer->IterateNodesOfType<UInterchangeFactoryBaseNode>([this](const FString& NodeUid, UInterchangeFactoryBaseNode* Node)
+		{
+			if (FactoryNodeClassesToDisabled.Contains(Node->GetClass()))
+			{
+				Node->SetEnabled(false);
+			}
+		});
+}
diff --git a/Engine/Plugins/Interchange/Editor/Source/Pipelines/Private/InterchangeEditorPipelineDetails.cpp b/Engine/Plugins/Interchange/Editor/Source/Pipelines/Private/InterchangeEditorPipelineDetails.cpp
--- a/Engine/Plugins/Interchange/Editor/Source/Pipelines/Private/InterchangeEditorPipelineDetails.cpp
+++ b/Engine/Plugins/Interchange/Editor/Source/Pipelines/Private/InterchangeEditorPipelineDetails.cpp
@@ -392,7 +392,7 @@ void FInterchangePipelineBaseDetailsCustomization::CustomizeDetails(IDetailLayou
 
 	const bool bAllowPropertyStatesEdition = InterchangePipeline->CanEditPropertiesStates();
 	const bool bIsReimportContext = InterchangePipeline->IsReimportContext();
-	const bool bIsBasicLayout = InterchangePipeline->IsBasicLayout();
+	const bool bIsShowEssentials = InterchangePipeline->IsShowEssentials();
 
 	TArray<FName> AllCategoryNames;
 	CachedDetailBuilder->GetCategoryNames(AllCategoryNames);
@@ -509,7 +509,7 @@ void FInterchangePipelineBaseDetailsCustomization::CustomizeDetails(IDetailLayou
 				bool IsLocked = false;
 				if (const FInterchangePipelinePropertyStates* PropertyStates = InterchangePipeline->GetPropertyStates(PropertyPath))
 				{
-					if (!PropertyStates->IsPropertyVisible(bIsReimportContext, bIsBasicLayout))
+					if (!PropertyStates->IsPropertyVisible(bIsReimportContext, bIsShowEssentials))
 					{
 						continue;
 					}
@@ -665,7 +665,7 @@ void FInterchangePipelineBaseDetailsCustomization::CustomizeDetails(IDetailLayou
 						SNew(STextBlock)
 						.Visibility(PipelineInternalEditionData ? EVisibility::Collapsed : EVisibility::Visible)
 						.Font(IDetailLayoutBuilder::GetDetailFont())
-						.Text(LOCTEXT("ShowWhenBasicLayoutText", "Basic Layout"))
+						.Text(LOCTEXT("ShowWhenShowEssentialsText", "Essential Layout"))
 					]
 					+ SHorizontalBox::Slot()
 					.AutoWidth()
@@ -679,7 +679,7 @@ void FInterchangePipelineBaseDetailsCustomization::CustomizeDetails(IDetailLayou
 						.UncheckedImage(FAppStyle::Get().GetBrush("Icons.Visible"))
 						.UncheckedHoveredImage(FAppStyle::Get().GetBrush("Icons.Visible"))
 						.UncheckedPressedImage(FAppStyle::Get().GetBrush("Icons.Visible"))
-						.ToolTipText(LOCTEXT("VisibleTooltipBasicLayout", "If true this property will be visible when displaying the interchange import dialog with basic layout."))
+						.ToolTipText(LOCTEXT("VisibleTooltipShowEssentials", "If true this property will be visible when displaying the interchange import dialog with basic layout."))
 						.OnCheckStateChanged_Lambda([InterchangePipelinePtr = InterchangePipeline, PropertyPath](ECheckBoxState CheckType)
 						{
 							if (!ensure(InterchangePipelinePtr.IsValid()))
@@ -688,7 +688,7 @@ void FInterchangePipelineBaseDetailsCustomization::CustomizeDetails(IDetailLayou
 							}
 							FScopedTransaction ScopedTransaction(LOCTEXT("TransactionvisibilityPropertiesToggleImport", "Toggle property visibility at import"), !GIsTransacting);
 							InterchangePipelinePtr->Modify();
-							InterchangePipelinePtr->FindOrAddPropertyStates(PropertyPath).SetPropertyBasicLayoutVisibility((CheckType != ECheckBoxState::Checked));
+							InterchangePipelinePtr->FindOrAddPropertyStates(PropertyPath).SetPropertyShowEssentialsVisibility((CheckType != ECheckBoxState::Checked));
 							InterchangePipelinePtr->PostEditChange();
 						})
 						.IsChecked_Lambda([InterchangePipelinePtr = InterchangePipeline, PropertyPath]()
@@ -699,7 +699,7 @@ void FInterchangePipelineBaseDetailsCustomization::CustomizeDetails(IDetailLayou
 							}
 							if (const FInterchangePipelinePropertyStates* PropertyStates = InterchangePipelinePtr->GetPropertyStates(PropertyPath))
 							{
-								return PropertyStates->IsPropertyVisibleInBasicLayout() ? ECheckBoxState::Unchecked : ECheckBoxState::Checked;
+								return PropertyStates->IsPropertyVisibleInShowEssentials() ? ECheckBoxState::Unchecked : ECheckBoxState::Checked;
 							}
 							return ECheckBoxState::Unchecked;
 						})
@@ -741,14 +741,14 @@ void FInterchangePipelineBaseDetailsCustomization::CustomizeDetails(IDetailLayou
 						.IsChecked_Lambda([InterchangePipelinePtr = InterchangePipeline, PropertyPath]()
 						{
 							constexpr bool bIsReimportContextLocal = false;
-							constexpr bool bIsBasicLayoutLocal = false;
+							constexpr bool bIsShowEssentialsLocal = false;
 							if (!InterchangePipelinePtr.IsValid())
 							{
 								return ECheckBoxState::Unchecked;
 							}
 							if (const FInterchangePipelinePropertyStates* PropertyStates = InterchangePipelinePtr->GetPropertyStates(PropertyPath))
 							{
-								return PropertyStates->IsPropertyVisible(bIsReimportContextLocal, bIsBasicLayoutLocal) ? ECheckBoxState::Unchecked : ECheckBoxState::Checked;
+								return PropertyStates->IsPropertyVisible(bIsReimportContextLocal, bIsShowEssentialsLocal) ? ECheckBoxState::Unchecked : ECheckBoxState::Checked;
 							}
 							return ECheckBoxState::Unchecked;
 						})
@@ -790,14 +790,14 @@ void FInterchangePipelineBaseDetailsCustomization::CustomizeDetails(IDetailLayou
 						.IsChecked_Lambda([InterchangePipelinePtr = InterchangePipeline, PropertyPath]()
 						{
 							constexpr bool bIsReimportContextLocal = true;
-							constexpr bool bIsBasicLayoutLocal = false;
+							constexpr bool bIsShowEssentialsLocal = false;
 							if (!InterchangePipelinePtr.IsValid())
 							{
 								return ECheckBoxState::Unchecked;
 							}
 							if (const FInterchangePipelinePropertyStates* PropertyStates = InterchangePipelinePtr->GetPropertyStates(PropertyPath))
 							{
-								return PropertyStates->IsPropertyVisible(bIsReimportContextLocal, bIsBasicLayoutLocal) ? ECheckBoxState::Unchecked : ECheckBoxState::Checked;
+								return PropertyStates->IsPropertyVisible(bIsReimportContextLocal, bIsShowEssentialsLocal) ? ECheckBoxState::Unchecked : ECheckBoxState::Checked;
 							}
 							return ECheckBoxState::Unchecked;
 						})
@@ -824,7 +824,7 @@ void FInterchangePipelineBaseDetailsCustomization::CustomizeDetails(IDetailLayou
 		}
 	}
 
-	if (!bIsBasicLayout)
+	if (!bIsShowEssentials)
 	{
 		AddExtraInformationSection();
 	}
diff --git a/Engine/Plugins/Interchange/Editor/Source/Pipelines/Private/InterchangeEditorPipelineStyle.h b/Engine/Plugins/Interchange/Editor/Source/Pipelines/Private/InterchangeEditorPipelineStyle.h
--- a/Engine/Plugins/Interchange/Editor/Source/Pipelines/Private/InterchangeEditorPipelineStyle.h
+++ b/Engine/Plugins/Interchange/Editor/Source/Pipelines/Private/InterchangeEditorPipelineStyle.h
@@ -6,6 +6,8 @@
 #include "Brushes/SlateImageBrush.h"
 #include "Styling/SlateStyle.h"
 #include "Styling/SlateStyleRegistry.h"
+#include "Styling/SlateTypes.h"
+#include "Styling/StyleColors.h"
 
 /**
  * Implements the visual style of the text asset editor UI.
@@ -20,6 +22,7 @@ public:
 		: FSlateStyleSet("InterchangeEditorPipelineStyle")
 	{
 		const FVector2D Icon16x16(16.0f, 16.0f);
+		const FVector2D Icon20x20(20.0f, 20.0f);
 
 		const FString BaseDir = IPluginManager::Get().FindPlugin("InterchangeEditor")->GetBaseDir();
 		SetContentRoot(BaseDir / TEXT("Content"));
@@ -41,6 +44,36 @@ public:
 		Set("PipelineConfigurationIcon.PipelineStack", PipelineStackBrush16);
 		FSlateImageBrush* PipelineStackDefaultBrush16 = new FSlateImageBrush(RootToContentDir(TEXT("Resources/Interchange_PipelineStackDefault_Icon_16"), TEXT(".png")), Icon16x16);
 		Set("PipelineConfigurationIcon.PipelineStackDefault", PipelineStackDefaultBrush16);
+		FSlateImageBrush* TranslatorSettingsBrush16 = new FSlateImageBrush(RootToContentDir(TEXT("Resources/Interchange_TranslatorSettings_16"), TEXT(".png")), Icon16x16);
+		Set("PipelineConfigurationIcon.TranslatorSettings", TranslatorSettingsBrush16);
+
+
+		//InterchangeAssetCard icons
+		FSlateImageBrush* AlertSolidWarningBrush20 = new FSlateImageBrush(RootToContentDir(TEXT("Old/Tiles/Outer/alertSolid"), TEXT(".png")), Icon20x20, FStyleColors::Warning);
+		Set("InterchangeAssetCardIcon.Warning", AlertSolidWarningBrush20);
+
+		FSlateImageBrush* SidePanelRightBrush20 = new FSlateImageBrush(RootToContentDir(TEXT("Resources/SidePanelRight"), TEXT(".png")), Icon20x20);
+		Set("PipelineConfigurationIcon.SidePanelRight", SidePanelRightBrush20);
+
+		Set("ImportSource.Dropdown.Border", new FSlateRoundedBoxBrush(FStyleColors::Dropdown.GetSpecifiedColor(), 4.f));
+		Set("AssetCardList.Background.Border", new FSlateRoundedBoxBrush(FStyleColors::Recessed.GetSpecifiedColor(), 4.f));
+		Set("AssetCard.Header.Border", new FSlateRoundedBoxBrush(FStyleColors::Dropdown.GetSpecifiedColor(), 4.f));
+		Set("AssetCard.Body.Border", new FSlateRoundedBoxBrush(FStyleColors::Panel.GetSpecifiedColor(), 0.f));
+
+
+		const FButtonStyle ResetSelectedPipelineButtonStyle = FButtonStyle()
+			.SetNormal(FSlateRoundedBoxBrush(FStyleColors::Dropdown, 4.0f))
+			.SetHovered(FSlateRoundedBoxBrush(FStyleColors::Dropdown, 4.0f))
+			.SetPressed(FSlateRoundedBoxBrush(FStyleColors::Dropdown, 4.0f))
+			.SetDisabled(FSlateNoResource())
+			.SetNormalForeground(FStyleColors::Foreground)
+			.SetHoveredForeground(FStyleColors::ForegroundHover)
+			.SetPressedForeground(FStyleColors::ForegroundHover)
+			.SetDisabledForeground(FStyleColors::Foreground)
+			.SetNormalPadding(FMargin(8.f, 4.f, 8.f, 4.f))
+			.SetPressedPadding(FMargin(8.f, 5.5f, 8.f, 2.5f));
+
+		Set("ButtonStyle.ResetSelectedPipeline", ResetSelectedPipelineButtonStyle);
 
 		FSlateStyleRegistry::RegisterSlateStyle(*this);
 	}
diff --git a/Engine/Plugins/Interchange/Editor/Source/Pipelines/Private/InterchangePipelineConfigurationGeneric.cpp b/Engine/Plugins/Interchange/Editor/Source/Pipelines/Private/InterchangePipelineConfigurationGeneric.cpp
--- a/Engine/Plugins/Interchange/Editor/Source/Pipelines/Private/InterchangePipelineConfigurationGeneric.cpp
+++ b/Engine/Plugins/Interchange/Editor/Source/Pipelines/Private/InterchangePipelineConfigurationGeneric.cpp
@@ -2,28 +2,79 @@
 #include "InterchangePipelineConfigurationGeneric.h"
 
 #include "Framework/Application/SlateApplication.h"
-#include "Nodes/InterchangeBaseNodeContainer.h"
 #include "Interfaces/IMainFrameModule.h"
+#include "Misc/ConfigCacheIni.h"
+#include "Misc/ConfigContext.h"
+#include "Nodes/InterchangeBaseNodeContainer.h"
 #include "SInterchangePipelineConfigurationDialog.h"
 #include "Widgets/SWindow.h"
 
 #include UE_INLINE_GENERATED_CPP_BY_NAME(InterchangePipelineConfigurationGeneric)
 
+
+namespace UE::Interchange::Private
+{
+	float GetMinSizeX() { return 550.f; }
+	float GetMinSizeY() { return 500.f; }
+
+	const FVector2D& DefaultClientSize()
+	{
+		static FVector2D DefaultClientSize(1000.0, 650.0);
+		return DefaultClientSize;
+	}
+
+	const FString& GetSizeXUid()
+	{
+		static FString SizeXStr = TEXT("SizeX");
+		return SizeXStr;
+	}
+
+	const FString& GetSizeYUid()
+	{
+		static FString SizeYStr = TEXT("SizeY");
+		return SizeYStr;
+	}
+	void RestoreClientSize(FVector2D& OutClientSize, const FString& WindowsUniqueID)
+	{
+		if (GConfig->DoesSectionExist(TEXT("InterchangeImportDialogOptions"), GEditorPerProjectIni))
+		{
+			FString SizeXUid = WindowsUniqueID + GetSizeXUid();
+			FString SizeYUid = WindowsUniqueID + GetSizeYUid();
+			GConfig->GetDouble(TEXT("InterchangeImportDialogOptions"), *SizeXUid, OutClientSize.X, GEditorPerProjectIni);
+			GConfig->GetDouble(TEXT("InterchangeImportDialogOptions"), *SizeYUid, OutClientSize.Y, GEditorPerProjectIni);
+		}
+	}
+	void BackupClientSize(const FVector2D& ClientSize, const FString& WindowsUniqueID)
+	{
+		FString SizeXUid = WindowsUniqueID + GetSizeXUid();
+		FString SizeYUid = WindowsUniqueID + GetSizeYUid();
+		GConfig->SetDouble(TEXT("InterchangeImportDialogOptions"), *SizeXUid, ClientSize.X, GEditorPerProjectIni);
+		GConfig->SetDouble(TEXT("InterchangeImportDialogOptions"), *SizeYUid, ClientSize.Y, GEditorPerProjectIni);
+	}
+}
+
 EInterchangePipelineConfigurationDialogResult UInterchangePipelineConfigurationGeneric::ShowPipelineConfigurationDialog(TArray<FInterchangeStackInfo>& PipelineStacks
 	, TArray<UInterchangePipelineBase*>& OutPipelines
 	, TWeakObjectPtr<UInterchangeSourceData> SourceData
 	, TWeakObjectPtr <UInterchangeTranslatorBase> Translator
 	, TWeakObjectPtr<UInterchangeBaseNodeContainer> BaseNodeContainer)
 {
+	using namespace UE::Interchange::Private;
 	//Create and show the graph inspector UI dialog
 	TSharedPtr<SWindow> ParentWindow;
 	if(IMainFrameModule* MainFrame = FModuleManager::LoadModulePtr<IMainFrameModule>("MainFrame"))
 	{
 		ParentWindow = MainFrame->GetParentWindow();
 	}
+	FString WindowDialogUid = TEXT("ImportContentDialog");
+	FVector2D ClientSize = DefaultClientSize();
+	RestoreClientSize(ClientSize, WindowDialogUid);
+
 	TSharedRef<SWindow> Window = SNew(SWindow)
-		.ClientSize(FVector2D(1000.f, 650.f))
-		.Title(NSLOCTEXT("Interchange", "PipelineConfigurationGenericTitleContent", "Interchange Pipeline Configuration (Import Content)"));
+		.ClientSize(ClientSize)
+		.MinWidth(GetMinSizeX())
+		.MinHeight(GetMinSizeY())
+		.Title(NSLOCTEXT("Interchange", "PipelineConfigurationGenericTitleContent", "Import Content"));
 	TSharedPtr<SInterchangePipelineConfigurationDialog> InterchangePipelineConfigurationDialog;
 
 	Window->SetContent
@@ -39,13 +90,19 @@ EInterchangePipelineConfigurationDialogResult UInterchangePipelineConfigurationG
 		.Translator(Translator)
 	);
 
+	Window->SetOnWindowClosed(FOnWindowClosed::CreateLambda([&WindowDialogUid](const TSharedRef<SWindow>& WindowMoved)
+		{
+			FVector2D LastClientSize = WindowMoved->GetClientSizeInScreen();
+			BackupClientSize(LastClientSize, WindowDialogUid);
+		}));
+
 	FSlateApplication::Get().AddModalWindow(Window, ParentWindow, false);
 	
 	if (InterchangePipelineConfigurationDialog->IsCanceled())
 	{
 		return EInterchangePipelineConfigurationDialogResult::Cancel;
 	}
-	
+
 	if (InterchangePipelineConfigurationDialog->IsImportAll())
 	{
 		return EInterchangePipelineConfigurationDialogResult::ImportAll;
@@ -60,15 +117,23 @@ EInterchangePipelineConfigurationDialogResult UInterchangePipelineConfigurationG
 	, TWeakObjectPtr <UInterchangeTranslatorBase> Translator
 	, TWeakObjectPtr<UInterchangeBaseNodeContainer> BaseNodeContainer)
 {
+	using namespace UE::Interchange::Private;
 	//Create and show the graph inspector UI dialog
 	TSharedPtr<SWindow> ParentWindow;
 	if (IMainFrameModule* MainFrame = FModuleManager::LoadModulePtr<IMainFrameModule>("MainFrame"))
 	{
 		ParentWindow = MainFrame->GetParentWindow();
 	}
+
+	FString WindowDialogUid = TEXT("ImportSceneDialog");
+	FVector2D ClientSize = DefaultClientSize();
+	RestoreClientSize(ClientSize, WindowDialogUid);
+
 	TSharedRef<SWindow> Window = SNew(SWindow)
-		.ClientSize(FVector2D(1000.f, 650.f))
-		.Title(NSLOCTEXT("Interchange", "PipelineConfigurationGenericTitleScene", "Interchange Pipeline Configuration (Import Scene)"));
+		.ClientSize(ClientSize)
+		.MinWidth(GetMinSizeX())
+		.MinHeight(GetMinSizeY())
+		.Title(NSLOCTEXT("Interchange", "PipelineConfigurationGenericTitleScene", "Import Scene"));
 	TSharedPtr<SInterchangePipelineConfigurationDialog> InterchangePipelineConfigurationDialog;
 
 	Window->SetContent
@@ -84,6 +149,12 @@ EInterchangePipelineConfigurationDialogResult UInterchangePipelineConfigurationG
 		.Translator(Translator)
 	);
 
+	Window->SetOnWindowClosed(FOnWindowClosed::CreateLambda([&WindowDialogUid](const TSharedRef<SWindow>& WindowMoved)
+		{
+			FVector2D LastClientSize = WindowMoved->GetClientSizeInScreen();
+			BackupClientSize(LastClientSize, WindowDialogUid);
+		}));
+
 	FSlateApplication::Get().AddModalWindow(Window, ParentWindow, false);
 
 	if (InterchangePipelineConfigurationDialog->IsCanceled())
@@ -106,15 +177,23 @@ EInterchangePipelineConfigurationDialogResult UInterchangePipelineConfigurationG
 	, TWeakObjectPtr<UInterchangeBaseNodeContainer> BaseNodeContainer
 	, TWeakObjectPtr <UObject> ReimportAsset)
 {
+	using namespace UE::Interchange::Private;
 	//Create and show the graph inspector UI dialog
 	TSharedPtr<SWindow> ParentWindow;
 	if (IMainFrameModule* MainFrame = FModuleManager::LoadModulePtr<IMainFrameModule>("MainFrame"))
 	{
 		ParentWindow = MainFrame->GetParentWindow();
 	}
+
+	FString WindowDialogUid = TEXT("ImportContentDialog");
+	FVector2D ClientSize = DefaultClientSize();
+	RestoreClientSize(ClientSize, WindowDialogUid);
+
 	TSharedRef<SWindow> Window = SNew(SWindow)
-		.ClientSize(FVector2D(1000.f, 650.f))
-		.Title(NSLOCTEXT("Interchange", "PipelineConfigurationGenericTitleReimportContent", "Interchange Pipeline Configuration (Reimport Content)"));
+		.ClientSize(ClientSize)
+		.MinWidth(GetMinSizeX())
+		.MinHeight(GetMinSizeY())
+		.Title(NSLOCTEXT("Interchange", "PipelineConfigurationGenericTitleReimportContent", "Reimport Content"));
 	TSharedPtr<SInterchangePipelineConfigurationDialog> InterchangePipelineConfigurationDialog;
 
 	Window->SetContent
@@ -131,6 +210,12 @@ EInterchangePipelineConfigurationDialogResult UInterchangePipelineConfigurationG
 		.Translator(Translator)
 	);
 
+	Window->SetOnWindowClosed(FOnWindowClosed::CreateLambda([&WindowDialogUid](const TSharedRef<SWindow>& WindowMoved)
+		{
+			FVector2D LastClientSize = WindowMoved->GetClientSizeInScreen();
+			BackupClientSize(LastClientSize, WindowDialogUid);
+		}));
+
 	FSlateApplication::Get().AddModalWindow(Window, ParentWindow, false);
 
 	if (InterchangePipelineConfigurationDialog->IsCanceled())
diff --git a/Engine/Plugins/Interchange/Editor/Source/Pipelines/Private/SInterchangeAssetCard.cpp b/Engine/Plugins/Interchange/Editor/Source/Pipelines/Private/SInterchangeAssetCard.cpp
new file mode 100644
--- /dev/null
+++ b/Engine/Plugins/Interchange/Editor/Source/Pipelines/Private/SInterchangeAssetCard.cpp
@@ -0,0 +1,249 @@
+// Copyright Epic Games, Inc. All Rights Reserved.
+#include "SInterchangeAssetCard.h"
+
+#include "GameFramework/Actor.h"
+#include "InterchangeEditorPipelineStyle.h"
+#include "Nodes/InterchangeBaseNodeContainer.h"
+#include "Nodes/InterchangeFactoryBaseNode.h"
+#include "SSimpleButton.h"
+#include "Styling/SlateBrush.h"
+#include "Styling/SlateIconFinder.h"
+#include "Styling/StyleColors.h"
+#include "Widgets/Images/SImage.h"
+#include "Widgets/Input/SCheckBox.h"
+#include "Widgets/Layout/SBorder.h"
+#include "Widgets/Layout/SBox.h"
+#include "Widgets/Layout/SSeparator.h"
+#include "Widgets/Layout/SSplitter.h"
+#include "Widgets/SBoxPanel.h"
+#include "Widgets/Text/STextBlock.h"
+#include "Widgets/Views/ITableRow.h"
+#include "Widgets/Views/SListView.h"
+#include "Widgets/Views/STableRow.h"
+#include "Widgets/Views/STableViewBase.h"
+
+#define LOCTEXT_NAMESPACE "InterchangeAssetCard"
+
+
+/************************************************************************/
+/* SInterchangeAssetCard Implementation                    */
+/************************************************************************/
+
+SInterchangeAssetCard::~SInterchangeAssetCard()
+{
+
+}
+
+void SInterchangeAssetCard::RefreshCard(UInterchangeBaseNodeContainer* InPreviewNodeContainer)
+{
+	check(InPreviewNodeContainer);
+	CardAssetCount = 0;
+	CardAssetToImportCount = 0;
+	CardAssetDisabledCount = 0;
+	CardTooltip.Reset();
+	//Query the data we need to update the card
+	InPreviewNodeContainer->IterateNodesOfType<UInterchangeFactoryBaseNode>([this](const FString& NodeUid, UInterchangeFactoryBaseNode* FactoryNode)
+		{
+			if (!FactoryNode || !FactoryNode->GetObjectClass() || !FactoryNode->GetObjectClass()->IsChildOf(AssetClass))
+			{
+				return;
+			}
+
+			CardAssetCount++;
+			if (!CardTooltip.IsEmpty())
+			{
+				CardTooltip += TEXT("\n");
+			}
+			if (FactoryNode->IsEnabled())
+			{
+				CardTooltip += TEXT("Import asset ") + FactoryNode->GetDisplayLabel();
+				CardAssetToImportCount++;
+			}
+			else
+			{
+				CardTooltip += TEXT("Ignore asset ") + FactoryNode->GetDisplayLabel();
+				CardAssetDisabledCount++;
+			}
+		});
+}
+
+void SInterchangeAssetCard::Construct(const FArguments& InArgs)
+{
+	//Make sure we have a valid node container
+	check(InArgs._PreviewNodeContainer);
+	AssetClass = InArgs._AssetClass;
+	check(AssetClass != nullptr);
+	ShouldImportAssetType = InArgs._ShouldImportAssetType;
+	check(ShouldImportAssetType.IsBound());
+	OnImportAssetTypeChanged = InArgs._OnImportAssetTypeChanged;
+	check(OnImportAssetTypeChanged.IsBound());
+
+	RefreshCard(InArgs._PreviewNodeContainer);
+
+	const FSlateBrush* CardAssetIcon = nullptr;
+	CardAssetIcon = FSlateIconFinder::FindIconBrushForClass(AssetClass);
+
+	TSharedPtr<SImage> CardIconWidget = SNew(SImage)
+		.Image(CardAssetIcon)
+		.Visibility(CardAssetIcon != FAppStyle::GetDefaultBrush() ? EVisibility::Visible : EVisibility::Collapsed);
+
+
+	FText CardImportText = FText::Format(LOCTEXT("CardImportText", "Import {0}"), FText::FromString(AssetClass->GetName()));
+
+	FName IconWarningName = "InterchangeAssetCardIcon.Warning";
+	const FSlateIcon IconWarning = FSlateIconFinder::FindIcon(IconWarningName);
+
+	const ISlateStyle* InterchangeEditorPipelineStyle = FSlateStyleRegistry::FindSlateStyle("InterchangeEditorPipelineStyle");
+	
+	const FSlateBrush* HeaderBorderBrush = nullptr;
+	const FSlateBrush* BodyBorderBrush = nullptr;
+	const FSlateBrush* BackgroundBorderBrush = nullptr;
+	if (InterchangeEditorPipelineStyle)
+	{
+		HeaderBorderBrush = InterchangeEditorPipelineStyle->GetBrush("AssetCard.Header.Border");
+		BodyBorderBrush = InterchangeEditorPipelineStyle->GetBrush("AssetCard.Body.Border");
+		BackgroundBorderBrush = InterchangeEditorPipelineStyle->GetBrush("AssetCardList.Background.Border");
+	}
+
+	ChildSlot
+	[
+		SNew(SVerticalBox)
+		.Visibility_Lambda([this]()
+			{
+				return CardAssetToImportCount > 0 ? EVisibility::Visible : EVisibility::Collapsed;
+			})
+		+ SVerticalBox::Slot()
+		.AutoHeight()
+		.Padding(0.0f, 4.0f, 0.0f, 0.0f)
+		[
+			SNew(SBorder)
+			.BorderImage(HeaderBorderBrush)
+			[
+				SNew(SHorizontalBox)
+				+ SHorizontalBox::Slot()
+				.AutoWidth()
+				.Padding(8.0f, 4.0f, 4.0f, 4.0f)
+				[
+					CardIconWidget.ToSharedRef()
+				]
+				+ SHorizontalBox::Slot()
+				.FillWidth(1.0f)
+				.Padding(4.0f, 4.0f, 8.0f, 4.0f)
+				[
+					SNew(STextBlock)
+					.Font(FAppStyle::GetFontStyle("NormalFont"))
+					.Text_Lambda([this]()
+						{
+							return FText::Format(LOCTEXT("CardLabelText", "{0} ({1} {2})"), FText::FromString(AssetClass->GetName()), FText::FromString(FString::FromInt(CardAssetToImportCount)), CardAssetToImportCount > 1 ? LOCTEXT("CardAssetsPlural", "assets") : LOCTEXT("CardAssetSingle", "asset"));
+						})
+					.ToolTipText_Lambda([this]()
+						{
+							return FText::FromString(CardTooltip);
+						})
+				]
+			]
+		]
+		+ SVerticalBox::Slot()
+		.Padding(0.0f, 0.0f, 0.0f, 4.0f)
+		.AutoHeight()
+		[
+			SNew(SBorder)
+			.BorderImage(BodyBorderBrush)
+			[
+				SNew(SHorizontalBox)
+				+ SHorizontalBox::Slot()
+				.FillWidth(0.6f)
+				.HAlign(HAlign_Left)
+				.Padding(8.0f, 4.0f)
+				[
+					SNew(SBox)
+					.Padding(0.0f)
+					[
+						SNew(STextBlock)
+						.Text(CardImportText)
+					]
+				]
+				+ SHorizontalBox::Slot()
+				.Padding(0.0f)
+				.AutoWidth()
+				[
+					SNew(SSeparator)
+					.Orientation(EOrientation::Orient_Vertical)
+					.Thickness(1.0f)
+				]
+				+ SHorizontalBox::Slot()
+				.FillWidth(0.4f)
+				.HAlign(HAlign_Left)
+				.Padding(8.0f, 4.0f)
+				[
+					SNew(SBox)
+					.Padding(0.)
+					[
+						SNew(SCheckBox)
+						.IsChecked_Lambda([this]()
+							{
+								if (ShouldImportAssetType.Execute())
+								{
+									return ECheckBoxState::Checked;
+								}
+								return ECheckBoxState::Unchecked;
+							})
+						.OnCheckStateChanged_Lambda([this](ECheckBoxState CheckBoxState)
+							{
+								OnImportAssetTypeChanged.Execute(CheckBoxState == ECheckBoxState::Checked);
+							})
+					]
+				]
+			]
+		]
+	];
+
+}
+
+
+void SInterchangeAssetCardList::Construct(const FArguments& InArgs)
+{
+	AssetCards = InArgs._AssetCards;
+	check(AssetCards);
+
+	AssetCardList = SNew(SListView<TSharedPtr<SInterchangeAssetCard>>)
+		.SelectionMode(ESelectionMode::None)
+		.ListItemsSource(AssetCards)
+		.OnGenerateRow(this, &SInterchangeAssetCardList::MakeAssetCardListRowWidget);
+
+	const ISlateStyle* InterchangeEditorPipelineStyle = FSlateStyleRegistry::FindSlateStyle("InterchangeEditorPipelineStyle");
+
+	const FSlateBrush* BackgroundBorderBrush = nullptr;
+	if (InterchangeEditorPipelineStyle)
+	{
+		BackgroundBorderBrush = InterchangeEditorPipelineStyle->GetBrush("AssetCardList.Background.Border");
+	}
+
+	ChildSlot
+	[
+		SNew(SBorder)
+		.Padding(8.0f, 4.0f)
+		.BorderImage(BackgroundBorderBrush)
+		[
+			AssetCardList.ToSharedRef()
+		]
+	];
+}
+
+void SInterchangeAssetCardList::RefreshList(UInterchangeBaseNodeContainer* InPreviewNodeContainer)
+{
+	for (TSharedPtr<SInterchangeAssetCard> AssetCard : *AssetCards)
+	{
+		AssetCard->RefreshCard(InPreviewNodeContainer);
+	}
+	AssetCardList->RequestListRefresh();
+}
+
+TSharedRef<ITableRow> SInterchangeAssetCardList::MakeAssetCardListRowWidget(TSharedPtr<SInterchangeAssetCard> InElement, const TSharedRef<STableViewBase>& OwnerTable)
+{
+	return SNew(STableRow<TSharedPtr<SInterchangeAssetCard>>, OwnerTable)
+		[
+			InElement.ToSharedRef()
+		];
+}
+#undef LOCTEXT_NAMESPACE
diff --git a/Engine/Plugins/Interchange/Editor/Source/Pipelines/Private/SInterchangeAssetCard.h b/Engine/Plugins/Interchange/Editor/Source/Pipelines/Private/SInterchangeAssetCard.h
new file mode 100644
--- /dev/null
+++ b/Engine/Plugins/Interchange/Editor/Source/Pipelines/Private/SInterchangeAssetCard.h
@@ -0,0 +1,82 @@
+// Copyright Epic Games, Inc. All Rights Reserved.
+
+#pragma once
+
+#include "Delegates/DelegateCombinations.h"
+#include "UObject/Class.h"
+#include "Widgets/Layout/SBox.h"
+#include "Widgets/SCompoundWidget.h"
+#include "Widgets/Views/ITableRow.h"
+#include "Widgets/Views/SListView.h"
+#include "Widgets/Views/STableViewBase.h"
+
+class UInterchangeBaseNodeContainer;
+
+DECLARE_DELEGATE_OneParam(FInterchangeCard_OnImportAssetTypeChanged, bool)
+DECLARE_DELEGATE_RetVal(bool, FInterchangeCard_ShouldImportAssetType)
+
+class SInterchangeAssetCard : public SCompoundWidget
+{
+public:
+	~SInterchangeAssetCard();
+
+	SLATE_BEGIN_ARGS(SInterchangeAssetCard)
+		: _PreviewNodeContainer(nullptr)
+		, _AssetClass(nullptr)
+		, _ShouldImportAssetType()
+		, _OnImportAssetTypeChanged()
+		{}
+		SLATE_ARGUMENT(UInterchangeBaseNodeContainer*, PreviewNodeContainer)
+		SLATE_ARGUMENT(UClass*, AssetClass)
+		SLATE_EVENT(FInterchangeCard_ShouldImportAssetType, ShouldImportAssetType)
+		SLATE_EVENT(FInterchangeCard_OnImportAssetTypeChanged, OnImportAssetTypeChanged)
+	SLATE_END_ARGS()
+
+public:
+	/** Construct this widget */
+	void Construct(const FArguments& InArgs);
+
+	void RefreshCard(UInterchangeBaseNodeContainer* InPreviewNodeContainer);
+
+protected:
+
+	int32 CardAssetCount = 0;
+	int32 CardAssetToImportCount = 0;
+	int32 CardAssetDisabledCount = 0;
+	FString CardTooltip;
+
+	/** The factory asset class so we can know which kind of Unreal asset this card is for. */
+	UClass* AssetClass = nullptr;
+
+	/** Delegate to invoke when we update the should import asset type UI. */
+	FInterchangeCard_ShouldImportAssetType ShouldImportAssetType;
+
+	/** Delegate to invoke when should import asset type value change. */
+	FInterchangeCard_OnImportAssetTypeChanged OnImportAssetTypeChanged;
+
+};
+
+class SInterchangeAssetCardList : public SCompoundWidget
+{
+public:
+
+	SLATE_BEGIN_ARGS(SInterchangeAssetCardList)
+		: _AssetCards()
+		{}
+		SLATE_ARGUMENT(TArray<TSharedPtr<SInterchangeAssetCard>>*, AssetCards)
+	SLATE_END_ARGS()
+
+public:
+	/** Construct this widget */
+	void Construct(const FArguments& InArgs);
+
+	void RefreshList(UInterchangeBaseNodeContainer* InPreviewNodeContainer);
+
+private:
+
+	TArray<TSharedPtr<SInterchangeAssetCard>>* AssetCards;
+
+	TSharedRef<ITableRow> MakeAssetCardListRowWidget(TSharedPtr<SInterchangeAssetCard> InElement, const TSharedRef<STableViewBase>& OwnerTable);
+
+	TSharedPtr <SListView<TSharedPtr<SInterchangeAssetCard>>> AssetCardList;
+};
diff --git a/Engine/Plugins/Interchange/Editor/Source/Pipelines/Private/SInterchangePipelineConfigurationDialog.cpp b/Engine/Plugins/Interchange/Editor/Source/Pipelines/Private/SInterchangePipelineConfigurationDialog.cpp
--- a/Engine/Plugins/Interchange/Editor/Source/Pipelines/Private/SInterchangePipelineConfigurationDialog.cpp
+++ b/Engine/Plugins/Interchange/Editor/Source/Pipelines/Private/SInterchangePipelineConfigurationDialog.cpp
@@ -11,36 +11,57 @@
 #include "GameFramework/Actor.h"
 #include "IDetailsView.h"
 #include "IDocumentation.h"
+#include "InterchangeCardsPipeline.h"
+#include "InterchangeEditorPipelineStyle.h"
 #include "InterchangeManager.h"
 #include "InterchangePipelineConfigurationBase.h"
 #include "InterchangeProjectSettings.h"
 #include "InterchangeTranslatorBase.h"
-#include "Nodes/InterchangeSourceNode.h"
 #include "Interfaces/IMainFrameModule.h"
+#include "Layout/Visibility.h"
 #include "Misc/App.h"
 #include "Misc/ConfigCacheIni.h"
 #include "Misc/ConfigContext.h"
 #include "Misc/ScopedSlowTask.h"
 #include "Nodes/InterchangeBaseNodeContainer.h"
+#include "Nodes/InterchangeSourceNode.h"
 #include "PropertyEditorModule.h"
 #include "SInterchangeGraphInspectorWindow.h"
 #include "SPrimaryButton.h"
 #include "Styling/SlateIconFinder.h"
 #include "Styling/StyleColors.h"
 #include "Widgets/Input/SCheckBox.h"
+#include "Widgets/Layout/SScrollBox.h"
+#include "Widgets/Layout/SSeparator.h"
 #include "Widgets/Input/STextComboBox.h"
 #include "Widgets/Layout/SUniformGridPanel.h"
-#include "Layout/Visibility.h"
 
 #define LOCTEXT_NAMESPACE "InterchangePipelineConfiguration"
 
-static bool GInterchangeDefaultBasicLayoutView = false;
-static FAutoConsoleVariableRef CCvarInterchangeEnableFBXImport(
-	TEXT("Interchange.FeatureFlags.Import.DefaultBasicLayoutView"),
-	GInterchangeDefaultBasicLayoutView,
-	TEXT("Whether the import dialog start by default in basic layout."),
+static bool GInterchangeDefaultShowEssentialsView = false;
+static FAutoConsoleVariableRef CCvarInterchangeDefaultShowEssentialsView(
+	TEXT("Interchange.FeatureFlags.Import.DefaultShowEssentialsView"),
+	GInterchangeDefaultShowEssentialsView,
+	TEXT("Whether the import dialog starts by default in essential pipeline properties layout."),
+	ECVF_Default);
+
+static bool GInterchangeDefaultShowSettings = false;
+static FAutoConsoleVariableRef CCvarInterchangeDefaultShowSettings(
+	TEXT("Interchange.FeatureFlags.Import.DefaultShowSettingsView"),
+	GInterchangeDefaultShowSettings,
+	TEXT("Whether the import dialog shows the settings by default. Settings mode is always shown if GInterchangeDefaultHideCardsView is true."),
 	ECVF_Default);
 
+static bool GInterchangeDefaultHideCardsView = true;
+static FAutoConsoleVariableRef CCvarInterchangeDefaultHideCardsView(
+	TEXT("Interchange.FeatureFlags.Import.DefaultHideCardsView"),
+	GInterchangeDefaultHideCardsView,
+	TEXT("Whether the import dialog should hide the basic cards view."),
+	ECVF_Default);
+
+constexpr double SplitterSize = 12.0;
+constexpr double AdvancedUIRatio = 1.2;
+
 const FName ReimportStackName = TEXT("ReimportPipeline");
 
 void SInterchangePipelineItem::Construct(
@@ -57,7 +78,7 @@ void SInterchangePipelineItem::Construct(
 	if (PipelineElementPtr.Get())
 	{
 		FString PipelineNameString = PipelineElement->DisplayName;
-		if (!PipelineElement->bBasicLayout)
+		if (!PipelineElement->bShowEssentials)
 		{
 			PipelineNameString += FString::Printf(TEXT(" (%s)"), *PipelineElementPtr->GetClass()->GetName());
 		}
@@ -154,6 +175,13 @@ SInterchangePipelineConfigurationDialog::~SInterchangePipelineConfigurationDialo
 	{
 		OwnerWindowPinned->GetOnWindowClosedEvent().RemoveAll(this);
 	}
+
+	if (PreviewNodeContainer)
+	{
+		PreviewNodeContainer->ClearFlags(RF_Standalone | RF_Public);
+		PreviewNodeContainer->ClearInternalFlags(EInternalObjectFlags::Async);
+		PreviewNodeContainer = nullptr;
+	}
 }
 
 // Pipelines are renamed with the reimport prefix to avoid conflicts with the duplicates of the original pipelines that end up in the same package.
@@ -277,6 +305,16 @@ void SInterchangePipelineConfigurationDialog::OnFinishedChangingProperties(const
 
 TSharedRef<SBox> SInterchangePipelineConfigurationDialog::SpawnPipelineConfiguration()
 {
+	const ISlateStyle* InterchangeEditorPipelineStyle = FSlateStyleRegistry::FindSlateStyle("InterchangeEditorPipelineStyle");
+
+	const FSlateBrush* ImportSourceBorderBrush = nullptr;
+	const FButtonStyle* ResetSelectPipelineButtonStyle = nullptr;
+	if (InterchangeEditorPipelineStyle)
+	{
+		ImportSourceBorderBrush = InterchangeEditorPipelineStyle->GetBrush("ImportSource.Dropdown.Border");
+		ResetSelectPipelineButtonStyle = &InterchangeEditorPipelineStyle->GetWidgetStyle<FButtonStyle>("ButtonStyle.ResetSelectedPipeline");
+	}
+
 	AvailableStacks.Reset();
 	TSharedPtr<FString> SelectedStack;
 	if (bReimport)
@@ -318,12 +356,12 @@ TSharedRef<SBox> SInterchangePipelineConfigurationDialog::SpawnPipelineConfigura
 						GeneratedPipeline->LoadSettings(Stack.StackName, bResetPreDialogTrue);
 						GeneratedPipeline->PreDialogCleanup(Stack.StackName);
 					}
-					GeneratedPipeline->SetBasicLayoutMode(bBasicLayout);
+					GeneratedPipeline->SetShowEssentialsMode(bShowEssentials);
 					if (bFilterOptions && BaseNodeContainer.IsValid())
 					{
 						GeneratedPipeline->FilterPropertiesFromTranslatedData(BaseNodeContainer.Get());
 					}
-					PipelineListViewItems.Add(MakeShareable(new FInterchangePipelineItemType{ GetPipelineDisplayName(DefaultPipeline), GeneratedPipeline, ReimportObject.Get(), BaseNodeContainer.Get(), SourceData.Get(), bBasicLayout }));
+					PipelineListViewItems.Add(MakeShareable(new FInterchangePipelineItemType{ GetPipelineDisplayName(DefaultPipeline), GeneratedPipeline, ReimportObject.Get(), BaseNodeContainer.Get(), SourceData.Get(), bShowEssentials }));
 				}
 			}
 			SelectedStack = StackNamePtr;
@@ -341,12 +379,13 @@ TSharedRef<SBox> SInterchangePipelineConfigurationDialog::SpawnPipelineConfigura
 
 	TSharedPtr<STextComboBox> TextComboBoxPtr;
 	//Only use a combo box if there is more then one stack
-	if (AvailableStacks.Num() > 1)
+	if (AvailableStacks.Num() > 0)
 	{
-		FText StackComboBoxTooltip = LOCTEXT("StackComboBoxTooltip", "Selected pipeline stack will be used for the current import. To change the pipeline stack used when automating or without dialog please change the default pipeline stacks in the project settings");
+		FText StackComboBoxTooltip = LOCTEXT("StackComboBoxTooltip", "Selected pipeline stack preset will be used for the current import. See the Interchange project settings to modify the pipeline stacks preset list.");
 		TextComboBoxPtr = SNew(STextComboBox)
 			.OptionsSource(&AvailableStacks)
 			.OnSelectionChanged(this, &SInterchangePipelineConfigurationDialog::OnStackSelectionChanged)
+			.ContentPadding(FMargin(0.0f, 2.0f))
 			.ToolTipText(StackComboBoxTooltip);
 		if (SelectedStack.IsValid())
 		{
@@ -354,12 +393,12 @@ TSharedRef<SBox> SInterchangePipelineConfigurationDialog::SpawnPipelineConfigura
 		}
 	}
 
-	FText CurrentStackText = LOCTEXT("CurrentStackText", "Choose Pipeline Stack: ");
+	FText CurrentStackText = LOCTEXT("CurrentStackText", "Stacks Preset");
 
 	TSharedPtr<SWidget> StackTextComboBox;
 	if (!TextComboBoxPtr.IsValid())
 	{
-		CurrentStackText = FText::Format(LOCTEXT("CurrentStackTextNoComboBox", "Pipeline Stack: {0}"), FText::FromName(CurrentStackName));
+		CurrentStackText = LOCTEXT("CurrentStackTextNoComboBox", "There is no pipeline stack preset available");
 		StackTextComboBox = SNew(SBox)
 		[
 			SNew(STextBlock)
@@ -368,27 +407,26 @@ TSharedRef<SBox> SInterchangePipelineConfigurationDialog::SpawnPipelineConfigura
 	}
 	else
 	{
-
-		StackTextComboBox = SNew(SHorizontalBox)
-		.Visibility_Lambda([this]()
-			{
-				return bBasicLayout ? EVisibility::Collapsed : EVisibility::Visible;
-			})
-		+ SHorizontalBox::Slot()
-		.VAlign(VAlign_Center)
-		.AutoWidth()
+		StackTextComboBox = SNew(SBox)
 		[
-			SNew(SBox)
+			SNew(SHorizontalBox)
+			+ SHorizontalBox::Slot()
+			.VAlign(VAlign_Center)
+			.Padding(0.0f, 0.0f, 8.0f, 0.0f)
+			.AutoWidth()
 			[
-				SNew(STextBlock)
-				.Text(CurrentStackText)
+				SNew(SBox)
+				[
+					SNew(STextBlock)
+					.Text(CurrentStackText)
+				]
+			]
+			+ SHorizontalBox::Slot()
+			.VAlign(VAlign_Center)
+			.AutoWidth()
+			[
+				TextComboBoxPtr.ToSharedRef()
 			]
-		]
-		+ SHorizontalBox::Slot()
-		.VAlign(VAlign_Center)
-		.AutoWidth()
-		[
-			TextComboBoxPtr.ToSharedRef()
 		];
 	}
 
@@ -401,13 +439,102 @@ TSharedRef<SBox> SInterchangePipelineConfigurationDialog::SpawnPipelineConfigura
 		[
 			SNew(SVerticalBox)
 			+ SVerticalBox::Slot()
-			.HAlign(HAlign_Left)
+			.Padding(0.0f, 8.0f)
 			.AutoHeight()
 			[
-				StackTextComboBox.ToSharedRef()
+				SNew(SHorizontalBox)
+				+ SHorizontalBox::Slot()
+				.AutoWidth()
+				[
+					StackTextComboBox.ToSharedRef()
+				]
+				+ SHorizontalBox::Slot()
+				.FillWidth(1.0f)
+				[
+					SNullWidget::NullWidget
+				]
+				+SHorizontalBox::Slot()
+				.AutoWidth()
+				.Padding(4.0f, 0.0f)
+				[
+					SNew(SBorder)
+					.BorderImage(ImportSourceBorderBrush)
+					.Padding(4.0f, 0.0f)
+					.ToolTipText(LOCTEXT("SInterchangePipelineConfigurationDialog_ShowEssentialsOptions_tooltip", "Display only essentials pipeline properties."))
+					[
+						SNew(SHorizontalBox)
+						+ SHorizontalBox::Slot()
+						.AutoWidth()
+						.VAlign(VAlign_Center)
+						.Padding(4.f, 0.f)
+						[
+							SNew(STextBlock)
+							.Text(LOCTEXT("SInterchangePipelineConfigurationDialog_ShowEssentialsOptions", "Essentials"))
+						]
+						+ SHorizontalBox::Slot()
+						.Padding(4.f, 0.f)
+						[
+							SNew(SCheckBox)
+							.IsChecked(this, &SInterchangePipelineConfigurationDialog::IsShowEssentialsEnabled)
+							.OnCheckStateChanged(this, &SInterchangePipelineConfigurationDialog::OnShowEssentialsChanged)
+						]
+					]
+				]
+				+SHorizontalBox::Slot()
+				.AutoWidth()
+				.Padding(4.0f, 0.0f)
+				[
+					SNew(SBorder)
+					.BorderImage(ImportSourceBorderBrush)
+					.Padding(4.0f, 0.0f)
+					.ToolTipText(LOCTEXT("SInterchangePipelineConfigurationDialog_FilterPipelineOptions_tooltip", "Filter the pipeline options using the source content data."))
+					[
+						SNew(SHorizontalBox)
+						+ SHorizontalBox::Slot()
+						.AutoWidth()
+						.VAlign(VAlign_Center)
+						.Padding(4.f, 0.f)
+						[
+							SNew(STextBlock)
+							.Text(LOCTEXT("SInterchangePipelineConfigurationDialog_FilterPipelineOptions", "Filter on Contents"))
+						]
+						+ SHorizontalBox::Slot()
+						.Padding(4.f, 0.f)
+						[
+							SNew(SCheckBox)
+							.IsChecked(this, &SInterchangePipelineConfigurationDialog::IsFilteringOptions)
+							.OnCheckStateChanged(this, &SInterchangePipelineConfigurationDialog::OnFilterOptionsChanged)
+						]
+					]
+				]
+				+SHorizontalBox::Slot()
+				.AutoWidth()
+				.VAlign(VAlign_Center)
+				.Padding(4.0f, 0.0f, 0.0f, 0.0f)
+				[
+					SNew(SButton)
+					.ButtonStyle(ResetSelectPipelineButtonStyle)
+					.IsEnabled_Lambda([this]()
+						{
+							return PipelinesListView->GetNumItemsSelected() == 1;
+						})
+					.Text(LOCTEXT("SInterchangePipelineConfigurationDialog_ResetToPipelineAsset", "Reset Selected Pipeline"))
+					.ToolTipText_Lambda([bReimportClosure = bReimport]()
+						{
+							if (bReimportClosure)
+							{
+								return LOCTEXT("SInterchangePipelineConfigurationDialog_ResetToPipelineAsset_TooltipReimport", "Reset the selected pipeline to the values used the last time this asset was imported.");
+							}
+							else
+							{
+								return LOCTEXT("SInterchangePipelineConfigurationDialog_ResetToPipelineAsset_Tooltip", "Reset the properties of the selected pipeline.");
+							}
+						})
+					.OnClicked(this, &SInterchangePipelineConfigurationDialog::OnResetToDefault)
+				]
 			]
 			+ SVerticalBox::Slot()
-			.Padding(0.0f, 5.0f)
+			.Padding(0.0f, 8.0f)
 			.AutoHeight()
 			[
 				SNew(SBox)
@@ -419,7 +546,6 @@ TSharedRef<SBox> SInterchangePipelineConfigurationDialog::SpawnPipelineConfigura
 			]
 		]
 		+ SVerticalBox::Slot()
-		.Padding(0.0f, 5.0f, 0.0f, 2.0f)
 		.FillHeight(1.0f)
 		[
 			SAssignNew(InspectorBox, SBox)
@@ -452,6 +578,111 @@ TSharedRef<SBox> SInterchangePipelineConfigurationDialog::SpawnPipelineConfigura
 	return PipelineConfigurationPanelBox;
 }
 
+TSharedRef<SBox> SInterchangePipelineConfigurationDialog::SpawnCardsConfiguration()
+{
+	const FSlateBrush* AdvanceSettingsIcon = FSlateIconFinder::FindIcon("PipelineConfigurationIcon.SidePanelRight").GetOptionalIcon();
+
+	CreateCardsViewList();
+
+	TSharedRef<SWidget> BodyWidget = CardViewList.IsValid() ? CardViewList.ToSharedRef() : SNullWidget::NullWidget;
+	TSharedRef<SBox> CardsConfigurationPanelBox = SNew(SBox)
+	[
+		SNew(SVerticalBox)
+		+ SVerticalBox::Slot()
+		.HAlign(HAlign_Fill)
+		.Padding(0.0f, 8.0f)
+		.AutoHeight()
+		[
+			SNew(SHorizontalBox)
+			+ SHorizontalBox::Slot()
+			.FillWidth(1.0f)
+			.VAlign(VAlign_Center)
+			.HAlign(HAlign_Left)
+			[
+				SNew(STextBlock)
+				.Text(LOCTEXT("SInterchangePipelineConfigurationDialog_AssetFoundText", "Assets Found In Source:"))
+			]
+			+ SHorizontalBox::Slot()
+			.AutoWidth()
+			.VAlign(VAlign_Center)
+			[
+				SNew(SCheckBox)
+				.Padding(FMargin(0.f, 4.f))
+				.Style(&FAppStyle::Get().GetWidgetStyle<FCheckBoxStyle>("ToggleButtonCheckBox"))
+				.Type(ESlateCheckBoxType::ToggleButton)
+				.IsChecked_Lambda([this]() { return bShowSettings ? ECheckBoxState::Checked : ECheckBoxState::Unchecked; })
+				.OnCheckStateChanged_Lambda([this](ECheckBoxState CheckState)
+					{
+						bShowSettings = CheckState == ECheckBoxState::Checked;
+						TSharedPtr<SWindow> DialogWindow = OwnerWindow.Pin();
+							
+						//Resize the window
+						FVector2D ClientSize = DialogWindow->GetClientSizeInScreen();
+						double MaxResize = 1500.0;
+						FWindowSizeLimits SizeLimits = DialogWindow->GetSizeLimits();
+						double MinimumWidth = OriginalMinWindowSize;
+						double NewSize = 0.0;
+						if (bShowSettings)
+						{
+							MinimumWidth = OriginalMinWindowSize + (OriginalMinWindowSize * AdvancedUIRatio) + SplitterSize;
+							//Increase the window size in X
+							double SizeRatio = (SplitAdvancedRatio / (1.0f - SplitAdvancedRatio));
+							//Add the splitter to the max size
+							NewSize = ClientSize.X + FMath::Min(MaxResize, (ClientSize.X * SizeRatio)) + SplitterSize;
+							ClientSize.X = FMath::CeilToDouble(NewSize);
+						}
+						else
+						{
+							MinimumWidth = OriginalMinWindowSize;
+							//Half the window size in X
+							double SizeRatio = (1.0f - SplitAdvancedRatio);
+							NewSize = (ClientSize.X - SplitterSize) * SizeRatio;
+							ClientSize.X = FMath::Floor(NewSize);
+						}
+
+						SizeLimits.SetMinWidth(static_cast<float>(MinimumWidth));
+						DialogWindow->SetSizeLimits(SizeLimits);
+						DialogWindow->Resize(ClientSize);
+
+						constexpr bool bStackSelectionChange = false;
+						RefreshStack(bStackSelectionChange);
+						GConfig->SetBool(TEXT("InterchangeImportDialogOptions"), TEXT("ShowSettings"), bShowSettings, GEditorPerProjectIni);
+					})
+				[
+					SNew(SHorizontalBox)
+					+ SHorizontalBox::Slot()
+					.AutoWidth()
+					.VAlign(VAlign_Center)
+					.Padding(0.0f, 0.0f, 8.0f, 0.0f)
+					[
+						SNew(SImage)
+						.Image(AdvanceSettingsIcon)
+					]
+					+ SHorizontalBox::Slot()
+					.AutoWidth()
+					.VAlign(VAlign_Center)
+					.Padding(0.0f, 0.0f, 8.0f, 0.0f)
+					[
+						SNew(STextBlock)
+						.Text(LOCTEXT("SInterchangePipelineConfigurationDialog_AdvanceSettingsButtonText", "Advanced Settings"))
+					]
+				]
+			]
+		]
+		+ SVerticalBox::Slot()
+		.Padding(0.0f, 8.0f)
+		.FillHeight(1.0f)
+		[
+			SNew(SScrollBox)
+			+ SScrollBox::Slot()
+			[
+				BodyWidget
+			]
+		]
+	];
+	return CardsConfigurationPanelBox;
+}
+
 void SInterchangePipelineConfigurationDialog::Construct(const FArguments& InArgs)
 {
 	LLM_SCOPE_BYNAME(TEXT("Interchange"));
@@ -488,150 +719,202 @@ void SInterchangePipelineConfigurationDialog::Construct(const FArguments& InArgs
 	check(OutPipelines);
 
 	check(OwnerWindow.IsValid());
-	if (TSharedPtr<SWindow> OwnerWindowPinned = OwnerWindow.Pin())
+	TSharedPtr<SWindow> OwnerWindowPinned = OwnerWindow.Pin();
+	if (OwnerWindowPinned.IsValid())
 	{
 		OwnerWindowPinned->GetOnWindowClosedEvent().AddRaw(this, &SInterchangePipelineConfigurationDialog::OnWindowClosed);
+		OriginalMinWindowSize = OwnerWindowPinned->GetSizeLimits().GetMinWidth().Get(0.0);
+		if (OriginalMinWindowSize < 1.0)
+		{
+			OriginalMinWindowSize = 350.0;
+		}
+		DeltaClientWindowSize = (OwnerWindowPinned->GetSizeInScreen() - OwnerWindowPinned->GetClientSizeInScreen()).X;
 	}
 
 	//Get the default layout when the user open the import dialog for the first time.
-	bBasicLayout = GInterchangeDefaultBasicLayoutView;
+	bShowEssentials = GInterchangeDefaultShowEssentialsView;
+	bShowCards = !GInterchangeDefaultHideCardsView;
+
+	bShowSettings = GInterchangeDefaultShowSettings || !bShowCards;
 
 	if (bReimport)
 	{
 		bFilterOptions = false;
 	}
-	else if(GConfig->DoesSectionExist(TEXT("InterchangeImportDialogOptions"), GEditorPerProjectIni))
+	
+	if(GConfig->DoesSectionExist(TEXT("InterchangeImportDialogOptions"), GEditorPerProjectIni))
+	{
+		if (!bReimport)
+		{
+			GConfig->GetBool(TEXT("InterchangeImportDialogOptions"), TEXT("FilterOptions"), bFilterOptions, GEditorPerProjectIni);
+		}
+		GConfig->GetBool(TEXT("InterchangeImportDialogOptions"), TEXT("ShowEssentials"), bShowEssentials, GEditorPerProjectIni);
+		GConfig->GetBool(TEXT("InterchangeImportDialogOptions"), TEXT("ShowSettings"), bShowSettings, GEditorPerProjectIni);
+		//Make sure settings are shown if we hide cards
+		if(!bShowCards)
+		{
+			bShowSettings = true;
+		}
+		GConfig->GetDouble(TEXT("InterchangeImportDialogOptions"), TEXT("SplitAdvancedRatio"), SplitAdvancedRatio, GEditorPerProjectIni);
+	}
+
+	//Make sure the windows is width enough to show all the ui part (cards and settings)
+	if (OwnerWindowPinned.IsValid())
+	{
+		FVector2D WidowsClientSize = OwnerWindowPinned->GetClientSizeInScreen();
+		FWindowSizeLimits SizeLimits = OwnerWindowPinned->GetSizeLimits();
+		double MinimumSizeX = 0.0;
+		if (bShowCards)
+		{
+			//Add show cards width
+			MinimumSizeX += OriginalMinWindowSize;
+		}
+		if (bShowSettings)
+		{
+			//Add settings width
+			MinimumSizeX += (OriginalMinWindowSize * AdvancedUIRatio);
+			
+			//If we have both we have to add the splitter width
+			if (bShowCards)
+			{
+				MinimumSizeX += SplitterSize;
+			}
+		}
+		SizeLimits.SetMinWidth(static_cast<float>(MinimumSizeX));
+		OwnerWindowPinned->SetSizeLimits(SizeLimits);
+		//Resize the window to respect the limits
+		OwnerWindowPinned->Resize(WidowsClientSize);
+	}
+
+	//SpawnPipelineConfiguration must always be call because it create the pipeline list from the project settings
+	TSharedRef<SBox> MainBodyAdvanced = SpawnPipelineConfiguration();
+	UpdatePipelineSupportedAssetClasses();
+	TSharedRef<SBox> MainBodyCardsConfiguration = SpawnCardsConfiguration();
+
+	const FSlateBrush* TranslatorSettingsIcon = FSlateIconFinder::FindIcon("PipelineConfigurationIcon.TranslatorSettings").GetOptionalIcon();
+
+	const ISlateStyle* InterchangeEditorPipelineStyle = FSlateStyleRegistry::FindSlateStyle("InterchangeEditorPipelineStyle");
+
+	const FSlateBrush* ImportSourceBorderBrush = nullptr;
+	if (InterchangeEditorPipelineStyle)
 	{
-		GConfig->GetBool(TEXT("InterchangeImportDialogOptions"), TEXT("FilterOptions"), bFilterOptions, GEditorPerProjectIni);
-		GConfig->GetBool(TEXT("InterchangeImportDialogOptions"), TEXT("BasicLayout"), bBasicLayout, GEditorPerProjectIni);
+		ImportSourceBorderBrush = InterchangeEditorPipelineStyle->GetBrush("ImportSource.Dropdown.Border");
 	}
 
 	this->ChildSlot
 	[
 		SNew(SBorder)
-		.Padding(FMargin(10.0f, 3.0f))
+		.Padding(FMargin(16.0f, 16.0f))
 		.BorderImage(FAppStyle::GetBrush("ToolPanel.DarkGroupBorder"))
 		[
 			SNew(SVerticalBox)
 			+ SVerticalBox::Slot()
 			.AutoHeight()
 			.HAlign(HAlign_Fill)
-			.Padding(0.0f, 2.0f)
+			.Padding(0.0f, 0.0f, 0.0f, 8.0f)
 			[
 				SNew(SHorizontalBox)
-				+SHorizontalBox::Slot()
+				+ SHorizontalBox::Slot()
 				.AutoWidth()
 				.VAlign(VAlign_Center)
-				.Padding(2.0f, 2.0f)
+				.Padding(8.0f, 0.0f, 8.0f, 0.0f)
 				[
-					SNew(SButton)
-					.Visibility_Lambda([this]()
-						{
-							return !TranslatorSettings ? EVisibility::Collapsed : EVisibility::Visible;
-						})
-					.ToolTipText(LOCTEXT("SInterchangePipelineConfigurationDialog_TranslatorSettings_Tooltip", "Edit translator project settings."))
-					.OnClicked(this, &SInterchangePipelineConfigurationDialog::OnEditTranslatorSettings)
-					[
-						SNew(SImage)
-							.Image(FAppStyle::GetBrush("Icons.Settings"))
-							.ColorAndOpacity(FSlateColor::UseForeground())
-					]
+					SNew(STextBlock)
+					.Text(LOCTEXT("SInterchangePipelineConfigurationDialog_SourceLabel", "Import Source"))
 				]
 				+SHorizontalBox::Slot()
 				.FillWidth(1.0f)
 				.VAlign(VAlign_Center)
-				.Padding(0.0f, 2.0f)
-				[
-					SNew(STextBlock)
-					.Text(this, &SInterchangePipelineConfigurationDialog::GetSourceDescription)
-				]
-				+SHorizontalBox::Slot()
-				.AutoWidth()
-				.HAlign(HAlign_Right)
-				.Padding(10.0f, 2.0f, 0.0f, 2.0f)
 				[
-					SNew(SHorizontalBox)
-					.ToolTipText(LOCTEXT("SInterchangePipelineConfigurationDialog_BasicLayoutOptions_tooltip", "Basic Layout display only the basic pipelines properties."))
-					+ SHorizontalBox::Slot()
-					.AutoWidth()
-					.HAlign(HAlign_Right)
-					.VAlign(VAlign_Center)
-					.Padding(4.f, 0.f)
+					SNew(SBorder)
+					.BorderImage(ImportSourceBorderBrush)
+					.Padding(0.0f, 4.0f)
 					[
-						SNew(STextBlock)
-						.Text(LOCTEXT("SInterchangePipelineConfigurationDialog_BasicLayoutOptions", "Basic Layout"))
-					]
-					+ SHorizontalBox::Slot()
-					.Padding(4.f, 0.f)
-					[
-						SNew(SCheckBox)
-						.IsChecked(this, &SInterchangePipelineConfigurationDialog::IsBasicLayoutEnabled)
-						.OnCheckStateChanged(this, &SInterchangePipelineConfigurationDialog::OnBasicLayoutChanged)
+						SNew(SHorizontalBox)
+						+ SHorizontalBox::Slot()
+						.FillWidth(1.0f)
+						.VAlign(VAlign_Center)
+						.Padding(8.0f, 0.0f, 0.0f, 0.0f)
+						[
+							SNew(STextBlock)
+							.Text(this, &SInterchangePipelineConfigurationDialog::GetSourceDescription)
+						]
+						+ SHorizontalBox::Slot()
+						.AutoWidth()
+						[
+							SNew(SSeparator)
+							.Orientation(EOrientation::Orient_Vertical)
+							.Thickness(1.0f)
+						]
+						+SHorizontalBox::Slot()
+						.AutoWidth()
+						.VAlign(VAlign_Center)
+						[
+							SNew(SButton)
+							.ButtonStyle(&FAppStyle::Get().GetWidgetStyle<FButtonStyle>("NoBorder"))
+							.ContentPadding(FMargin(0.0f))
+							.Visibility_Lambda([this]()
+								{
+									return !TranslatorSettings ? EVisibility::Collapsed : EVisibility::Visible;
+								})
+							.ToolTipText(LOCTEXT("SInterchangePipelineConfigurationDialog_TranslatorSettings_Tooltip", "Edit translator project settings."))
+							.OnClicked(this, &SInterchangePipelineConfigurationDialog::OnEditTranslatorSettings)
+							[
+								SNew(SImage)
+								.Image(TranslatorSettingsIcon)
+							]
+						]
 					]
 				]
-				+SHorizontalBox::Slot()
-				.AutoWidth()
-				.HAlign(HAlign_Right)
-				.Padding(10.0f, 2.0f, 0.0f, 2.0f)
+			]
+			+ SVerticalBox::Slot()
+			.FillHeight(1.0f)
+			[
+				SAssignNew(CardsAndAdvancedSplitter, SSplitter)
+				.OnSplitterFinishedResizing_Lambda([this]()
+					{
+						SplitAdvancedRatio = (1.0 - static_cast<double>(CardsAndAdvancedSplitter->SlotAt(0).GetSizeValue()));
+						GConfig->SetDouble(TEXT("InterchangeImportDialogOptions"), TEXT("SplitAdvancedRatio"), SplitAdvancedRatio, GEditorPerProjectIni);
+					})
+				+ SSplitter::Slot()
+				.MinSize(static_cast<float>(OriginalMinWindowSize - DeltaClientWindowSize))
+				.Value(1.0f - static_cast<float>(SplitAdvancedRatio))
 				[
-					SNew(SHorizontalBox)
-					.ToolTipText(LOCTEXT("SInterchangePipelineConfigurationDialog_FilterPipelineOptions_tooltip", "Filter the pipeline options using the source content data."))
+					SNew(SBox)
 					.Visibility_Lambda([this]()
 						{
-							return bReimport ? EVisibility::Collapsed : EVisibility::Visible;
+							return bShowCards ? EVisibility::Visible : EVisibility::Collapsed;
+						})
+					.Padding_Lambda([this]()
+						{
+							return (bShowSettings && bShowCards) ? FMargin(0.0f, 0.0f, 20.0f, 0.0f) : FMargin(0.0f, 0.0f, 0.0f, 0.0f);
 						})
-					+ SHorizontalBox::Slot()
-					.AutoWidth()
-					.HAlign(HAlign_Right)
-					.VAlign(VAlign_Center)
-					.Padding(4.f, 0.f)
-					[
-						SNew(STextBlock)
-						.Text(LOCTEXT("SInterchangePipelineConfigurationDialog_FilterPipelineOptions", "Filter on Contents"))
-					]
-					+ SHorizontalBox::Slot()
-					.Padding(4.f, 0.f)
 					[
-						SNew(SCheckBox)
-						.IsChecked(this, &SInterchangePipelineConfigurationDialog::IsFilteringOptions)
-						.OnCheckStateChanged(this, &SInterchangePipelineConfigurationDialog::OnFilterOptionsChanged)
+						MainBodyCardsConfiguration
 					]
 				]
-				+SHorizontalBox::Slot()
-				.AutoWidth()
-				.HAlign(HAlign_Right)
-				.Padding(2.0f, 2.0f, 0.0f, 2.0f)
+				+ SSplitter::Slot()
+				.MinSize(static_cast<float>((OriginalMinWindowSize * AdvancedUIRatio) - DeltaClientWindowSize))
+				.Value(static_cast<float>(SplitAdvancedRatio))
 				[
-					SNew(SButton)
-					.HAlign(HAlign_Center)
-					.IsEnabled_Lambda([this]()
+					SNew(SBox)
+					.Visibility_Lambda([this]()
 						{
-							return PipelinesListView->GetNumItemsSelected() == 1;
+							return bShowSettings ? EVisibility::Visible : EVisibility::Collapsed;
 						})
-					.Text(LOCTEXT("SInterchangePipelineConfigurationDialog_ResetToPipelineAsset", "Reset Selected Pipeline"))
-					.ToolTipText_Lambda([bReimportClosure = bReimport]()
+					.Padding_Lambda([this]()
 						{
-							if (bReimportClosure)
-							{
-								return LOCTEXT("SInterchangePipelineConfigurationDialog_ResetToPipelineAsset_TooltipReimport", "Reset the selected pipeline to the values used the last time this asset was imported.");
-							}
-							else
-							{
-								return LOCTEXT("SInterchangePipelineConfigurationDialog_ResetToPipelineAsset_Tooltip", "Reset the properties of the selected pipeline.");
-							}
+							return bShowCards ? FMargin(8.0f, 0.0f, 0.0f, 0.0f) : FMargin(0.0f, 0.0f, 0.0f, 0.0f);
 						})
-					.OnClicked(this, &SInterchangePipelineConfigurationDialog::OnResetToDefault)
+					.Padding(8.0f, 0.0f, 0.0f, 0.0f)
+					[
+						MainBodyAdvanced
+					]
 				]
 			]
 			+ SVerticalBox::Slot()
-			.FillHeight(1.0f)
-			[
-				SpawnPipelineConfiguration()
-			]
-			+ SVerticalBox::Slot()
 			.AutoHeight()
-			.HAlign(HAlign_Right)
-			.Padding(0.0f, 2.0f)
+			.Padding(0.0f, 8.0f, 0.0f, 0.0f)
 			[
 				SNew(SHorizontalBox)
 				+ SHorizontalBox::Slot()
@@ -673,21 +956,19 @@ void SInterchangePipelineConfigurationDialog::Construct(const FArguments& InArgs
 					]
 				]
 				+ SHorizontalBox::Slot()
-				.AutoWidth()
-				.Padding(4.f, 0.f) 
+				.FillWidth(1.0f)
 				[
-					SNew(SPrimaryButton)
-					.Icon(this, &SInterchangePipelineConfigurationDialog::GetImportButtonIcon)
-					.Text(LOCTEXT("InspectorGraphWindow_Import", "Import"))
-					.ToolTipText(this, &SInterchangePipelineConfigurationDialog::GetImportButtonTooltip)
-					.IsEnabled(this, &SInterchangePipelineConfigurationDialog::IsImportButtonEnabled)
-					.OnClicked(this, &SInterchangePipelineConfigurationDialog::OnCloseDialog, ECloseEventType::Import)
+					SNullWidget::NullWidget
 				]
 				+ SHorizontalBox::Slot()
 				.AutoWidth()
 				.Padding(4.f, 0.f)
 				[
 					SNew(SButton)
+					.Visibility_Lambda([this]()
+						{
+							return bShowSettings ? EVisibility::Visible : EVisibility::Collapsed;
+						})
 					.HAlign(HAlign_Center)
 					.Text(LOCTEXT("InspectorGraphWindow_Preview", "Preview..."))
 					.OnClicked(this, &SInterchangePipelineConfigurationDialog::OnPreviewImport)
@@ -695,6 +976,17 @@ void SInterchangePipelineConfigurationDialog::Construct(const FArguments& InArgs
 				+ SHorizontalBox::Slot()
 				.AutoWidth()
 				.Padding(4.f, 0.f) 
+				[
+					SNew(SPrimaryButton)
+					.Icon(this, &SInterchangePipelineConfigurationDialog::GetImportButtonIcon)
+					.Text(LOCTEXT("InspectorGraphWindow_Import", "Import"))
+					.ToolTipText(this, &SInterchangePipelineConfigurationDialog::GetImportButtonTooltip)
+					.IsEnabled(this, &SInterchangePipelineConfigurationDialog::IsImportButtonEnabled)
+					.OnClicked(this, &SInterchangePipelineConfigurationDialog::OnCloseDialog, ECloseEventType::Import)
+				]
+				+ SHorizontalBox::Slot()
+				.AutoWidth()
+				.Padding(4.f, 0.f) 
 				[
 					SNew(SButton)
 					.HAlign(HAlign_Center)
@@ -762,17 +1054,9 @@ const FSlateBrush* SInterchangePipelineConfigurationDialog::GetImportButtonIcon(
 FText SInterchangePipelineConfigurationDialog::GetSourceDescription() const
 {
 	FText ActionDescription;
-	if (bReimport)
-	{
-		ActionDescription = LOCTEXT("GetSourceDescription_Reimport", "Reimport");
-	}
-	else
-	{
-		ActionDescription = LOCTEXT("GetSourceDescription_Import", "Import");
-	}
 	if (SourceData.IsValid())
 	{
-		ActionDescription = FText::Format(LOCTEXT("GetSourceDescription", "{0} source {1}"), ActionDescription, FText::FromString(SourceData->GetFilename()));
+		ActionDescription = FText::FromString(SourceData->GetFilename());
 	}
 	return ActionDescription;
 }
@@ -815,7 +1099,7 @@ FReply SInterchangePipelineConfigurationDialog::OnResetToDefault()
 								if (UInterchangePipelineBase* GeneratedPipeline = UE::Interchange::GeneratePipelineInstance(DefaultPipeline))
 								{
 									GeneratedPipeline->TransferAdjustSettings(DefaultPipeline);
-									GeneratedPipeline->SetBasicLayoutMode(bBasicLayout);
+									GeneratedPipeline->SetShowEssentialsMode(bShowEssentials);
 									if(bFilterOptions && BaseNodeContainer.IsValid())
 									{
 										GeneratedPipeline->FilterPropertiesFromTranslatedData(BaseNodeContainer.Get());
@@ -835,6 +1119,10 @@ FReply SInterchangePipelineConfigurationDialog::OnResetToDefault()
 		}
 	}
 	SetEditPipeline(PipelineToEdit);
+
+	//Update the cards
+	RefreshCardsViewList();
+
 	return Result;
 }
 
@@ -910,6 +1198,11 @@ void SInterchangePipelineConfigurationDialog::ClosePipelineConfiguration(const E
 		{
 			OutPipelines->Add(PipelineElement->Pipeline);
 		}
+		if (UInterchangeCardsPipeline* InterchangeCardsPipeline = GenerateTransientCardsPipeline())
+		{
+			//Add the cards pipeline if valid
+			OutPipelines->Add(InterchangeCardsPipeline);
+		}
 	}
 
 	//Save the settings only if its not a re-import
@@ -994,23 +1287,29 @@ void SInterchangePipelineConfigurationDialog::RefreshStack(bool bStackSelectionC
 						GeneratedPipeline->PreDialogCleanup(Stack.StackName);
 					}
 				}
-				GeneratedPipeline->SetBasicLayoutMode(bBasicLayout);
+				GeneratedPipeline->SetShowEssentialsMode(bShowEssentials);
 				if (bFilterOptions && BaseNodeContainer.IsValid())
 				{
 					GeneratedPipeline->FilterPropertiesFromTranslatedData(BaseNodeContainer.Get());
 				}
-				PipelineListViewItems.Add(MakeShareable(new FInterchangePipelineItemType{ GetPipelineDisplayName(DefaultPipeline), GeneratedPipeline, ReimportObject.Get(), BaseNodeContainer.Get(), SourceData.Get(), bBasicLayout }));
+				PipelineListViewItems.Add(MakeShareable(new FInterchangePipelineItemType{ GetPipelineDisplayName(DefaultPipeline), GeneratedPipeline, ReimportObject.Get(), BaseNodeContainer.Get(), SourceData.Get(), bShowEssentials }));
 			}
 		}
 	}
 
-	//Select the first pipeline
-	if (PipelineListViewItems.Num() > 0)
+	if (bShowSettings)
 	{
-		CurrentPipelineIndex = PipelineListViewItems.IsValidIndex(CurrentPipelineIndex) ? CurrentPipelineIndex : 0;
-		PipelinesListView->SetSelection(PipelineListViewItems[CurrentPipelineIndex], ESelectInfo::Direct);
+		//Select the first pipeline
+		if (PipelineListViewItems.Num() > 0)
+		{
+			CurrentPipelineIndex = PipelineListViewItems.IsValidIndex(CurrentPipelineIndex) ? CurrentPipelineIndex : 0;
+			PipelinesListView->SetSelection(PipelineListViewItems[CurrentPipelineIndex], ESelectInfo::Direct);
+		}
+		PipelinesListView->RequestListRefresh();
 	}
-	PipelinesListView->RequestListRefresh();
+
+	//Update the cards
+	RefreshCardsViewList();
 }
 
 void SInterchangePipelineConfigurationDialog::OnStackSelectionChanged(TSharedPtr<FString> String, ESelectInfo::Type)
@@ -1063,6 +1362,155 @@ void SInterchangePipelineConfigurationDialog::OnPipelineSelectionChanged(TShared
 		GConfig->SetString(TEXT("InterchangeSelectPipeline"), *KeyName, *CurrentPipelineName, GEditorPerProjectIni);
 	}
 }
+void SInterchangePipelineConfigurationDialog::UpdatePipelineSupportedAssetClasses()
+{
+	PipelineSupportAssetClasses.Reset();
+	for (TSharedPtr<FInterchangePipelineItemType> PipelineItem : PipelineListViewItems)
+	{
+		TArray<UClass*> PipelineSupportedClasses;
+		PipelineItem->Pipeline->GetSupportAssetClasses(PipelineSupportedClasses);
+		for (UClass* AssetClass : PipelineSupportedClasses)
+		{
+			if (!PipelineSupportAssetClasses.Contains(AssetClass))
+			{
+				PipelineSupportAssetClasses.Add(AssetClass);
+			}
+		}
+	}
+}
+
+void SInterchangePipelineConfigurationDialog::UpdateEnableDataPerFactoryNodeClass()
+{
+	if (!ensure(PreviewNodeContainer))
+	{
+		return;
+	}
+	TArray<UClass*> ValidCardClasses;
+	PreviewNodeContainer->IterateNodesOfType<UInterchangeFactoryBaseNode>([this, &ValidCardClasses](const FString& NodeUid, UInterchangeFactoryBaseNode* FactoryNode)
+		{
+			UClass* NodeObjectClass = FactoryNode->GetObjectClass();
+			// Don't add factory node class twice.
+			// Don't add factory node class that don't have a valid object class
+			// Don't add factory node class that are not a main asset from pipelines
+			if (!ValidCardClasses.Contains(FactoryNode->GetClass())
+				&& NodeObjectClass)
+			{
+				bool bClassIsSupported = false;
+				for (UClass* SupportedClass : PipelineSupportAssetClasses)
+				{
+					if (NodeObjectClass->IsChildOf(SupportedClass))
+					{
+						bClassIsSupported = true;
+						break;
+					}
+				}
+				if (bClassIsSupported)
+				{
+					ValidCardClasses.Add(FactoryNode->GetClass());
+					if (!EnableDataPerFactoryNodeClass.Contains(FactoryNode->GetClass()))
+					{
+						//Add a new entry that is enabled by default
+						EnableDataPerFactoryNodeClass.FindOrAdd(FactoryNode->GetClass()).ObjectClass = FactoryNode->GetObjectClass();
+					}
+				}
+			}
+		});
+
+	//Remove class card that do not exist anymore
+	TArray<UClass*> CardsClassesToRemove;
+	for (TPair<UClass*, FFactoryNodeEnabledData>& FactoryNodeClassAndEnableStatus : EnableDataPerFactoryNodeClass)
+	{
+		if (!ValidCardClasses.Contains(FactoryNodeClassAndEnableStatus.Key))
+		{
+			CardsClassesToRemove.Add(FactoryNodeClassAndEnableStatus.Key);
+		}
+	}
+	for (UClass* CardToRemove : CardsClassesToRemove)
+	{
+		EnableDataPerFactoryNodeClass.Remove(CardToRemove);
+	}
+}
+
+void SInterchangePipelineConfigurationDialog::FillAssetCardsList()
+{
+	//Update the pipeline supported asset class
+	UpdatePipelineSupportedAssetClasses();
+
+	//Update the preview container
+	constexpr bool bUpdateCardsTrue = true;
+	UpdatePreviewContainer(bUpdateCardsTrue);
+	if (ensure(PreviewNodeContainer))
+	{
+		UpdateEnableDataPerFactoryNodeClass();
+		AssetCards.Reset();
+		for (TPair<UClass*, FFactoryNodeEnabledData>& FactoryNodeClassAndEnableStatus : EnableDataPerFactoryNodeClass)
+		{
+			UClass* FactoryNodeClass = FactoryNodeClassAndEnableStatus.Key;
+			UClass* AssetClass = FactoryNodeClassAndEnableStatus.Value.ObjectClass;
+			if (AssetClass)
+			{
+				TSharedPtr<SInterchangeAssetCard> AssetCard = SNew(SInterchangeAssetCard)
+					.PreviewNodeContainer(PreviewNodeContainer)
+					.AssetClass(AssetClass)
+					.ShouldImportAssetType_Lambda([this, FactoryNodeClass]()
+						{
+							return EnableDataPerFactoryNodeClass.FindChecked(FactoryNodeClass).bEnable;
+						})
+					.OnImportAssetTypeChanged_Lambda([this, FactoryNodeClass](bool bNewEnabledValue)
+						{
+							EnableDataPerFactoryNodeClass.FindChecked(FactoryNodeClass).bEnable = bNewEnabledValue;
+						});
+
+				AssetCards.Add(AssetCard);
+			}
+		}
+	}
+}
+
+void SInterchangePipelineConfigurationDialog::CreateCardsViewList()
+{
+	FillAssetCardsList();
+	if (AssetCards.IsEmpty())
+	{
+		CardViewList = nullptr;
+	}
+	else
+	{
+		CardViewList = SNew(SInterchangeAssetCardList).AssetCards(&AssetCards);
+	}
+}
+
+void SInterchangePipelineConfigurationDialog::RefreshCardsViewList()
+{
+	FillAssetCardsList();
+	if (CardViewList.IsValid())
+	{
+		CardViewList->RefreshList(PreviewNodeContainer);
+	}
+}
+
+UInterchangeCardsPipeline* SInterchangePipelineConfigurationDialog::GenerateTransientCardsPipeline() const
+{
+	UInterchangeCardsPipeline* InterchangeCardsPipeline = nullptr;
+	if (!bReimport)
+	{
+		TArray<UClass*> DisabledNodeClasses;
+		for (const TPair<UClass*, FFactoryNodeEnabledData>& FactoryNodeClassAndEnableStatus : EnableDataPerFactoryNodeClass)
+		{
+			if(!FactoryNodeClassAndEnableStatus.Value.bEnable)
+			{
+				DisabledNodeClasses.Add(FactoryNodeClassAndEnableStatus.Key);
+			}
+		}
+
+		if (!DisabledNodeClasses.IsEmpty())
+		{
+			InterchangeCardsPipeline = NewObject<UInterchangeCardsPipeline>();
+			InterchangeCardsPipeline->SetDisabledFactoryNodes(DisabledNodeClasses);
+		}
+	}
+	return InterchangeCardsPipeline;
+}
 
 void SInterchangePipelineConfigurationDialog::OnFilterOptionsChanged(ECheckBoxState CheckState)
 {
@@ -1080,30 +1528,35 @@ void SInterchangePipelineConfigurationDialog::OnFilterOptionsChanged(ECheckBoxSt
 	GConfig->SetBool(TEXT("InterchangeImportDialogOptions"), TEXT("FilterOptions"), bFilterOptions, GEditorPerProjectIni);
 }
 
-void SInterchangePipelineConfigurationDialog::OnBasicLayoutChanged(ECheckBoxState CheckState)
+void SInterchangePipelineConfigurationDialog::OnShowEssentialsChanged(ECheckBoxState CheckState)
 {
 	bool bNewCheckValue = CheckState == ECheckBoxState::Checked ? true : false;
-	if (bNewCheckValue == bBasicLayout)
+	if (bNewCheckValue == bShowEssentials)
 	{
 		//Check state did not change
 		return;
 	}
-	bBasicLayout = bNewCheckValue;
+	bShowEssentials = bNewCheckValue;
 	//Refresh the pipeline
 	constexpr bool bStackSelectionChange = false;
 	RefreshStack(bStackSelectionChange);
 
-	GConfig->SetBool(TEXT("InterchangeImportDialogOptions"), TEXT("BasicLayout"), bBasicLayout, GEditorPerProjectIni);
+	GConfig->SetBool(TEXT("InterchangeImportDialogOptions"), TEXT("ShowEssentials"), bShowEssentials, GEditorPerProjectIni);
 }
 
-FReply SInterchangePipelineConfigurationDialog::OnPreviewImport() const
+void SInterchangePipelineConfigurationDialog::UpdatePreviewContainer(bool bUpdateCards) const
 {
 	auto ClearObjectFlags = [](UObject* Obj)
 		{
 			Obj->ClearFlags(RF_Standalone | RF_Public);
 			Obj->ClearInternalFlags(EInternalObjectFlags::Async);
 		};
-	UInterchangeBaseNodeContainer* DuplicateBaseNodeContainer = DuplicateObject<UInterchangeBaseNodeContainer>(BaseNodeContainer.Get(), GetTransientPackage());
+	if (PreviewNodeContainer)
+	{
+		ClearObjectFlags(PreviewNodeContainer);
+		PreviewNodeContainer = nullptr;
+	}
+	PreviewNodeContainer = DuplicateObject<UInterchangeBaseNodeContainer>(BaseNodeContainer.Get(), GetTransientPackage());
 
 	TArray<UInterchangeSourceData*> SourceDatas;
 	SourceDatas.Add(SourceData.Get());
@@ -1113,18 +1566,26 @@ FReply SInterchangePipelineConfigurationDialog::OnPreviewImport() const
 	for (int32 PipelineIndex = 0; PipelineIndex < PipelineListViewItems.Num(); ++PipelineIndex)
 	{
 		const TSharedPtr<FInterchangePipelineItemType> PipelineItem = PipelineListViewItems[PipelineIndex];
-		
+
 		//Duplicate the pipeline because ScriptedExecutePipeline is not const
 		if (UInterchangePipelineBase* DuplicatedPipeline = DuplicateObject<UInterchangePipelineBase>(PipelineItem->Pipeline, GetTransientPackage()))
 		{
 			DuplicatedPipeline->SetResultsContainer(Results);
-			DuplicatedPipeline->ScriptedExecutePipeline(DuplicateBaseNodeContainer, SourceDatas, FString());
+			DuplicatedPipeline->ScriptedExecutePipeline(PreviewNodeContainer, SourceDatas, FString());
 			ClearObjectFlags(DuplicatedPipeline);
 		}
 	}
 
-	
-	DuplicateBaseNodeContainer->IterateNodesOfType<UInterchangeFactoryBaseNode>([ClosureReimportObject = ReimportObject](const FString& NodeUid, UInterchangeFactoryBaseNode* Node)
+	if (!bUpdateCards)
+	{
+		//If we do not update cards execute the cards pipeline since its a final preview
+		if (UInterchangeCardsPipeline* InterchangeCardsPipeline = GenerateTransientCardsPipeline())
+		{
+			InterchangeCardsPipeline->ScriptedExecutePipeline(PreviewNodeContainer, SourceDatas, FString());
+		}
+	}
+
+	PreviewNodeContainer->IterateNodesOfType<UInterchangeFactoryBaseNode>([ClosureReimportObject = ReimportObject](const FString& NodeUid, UInterchangeFactoryBaseNode* Node)
 		{
 
 			//Set all node in preview mode so hide the internal data attributes
@@ -1137,6 +1598,26 @@ FReply SInterchangePipelineConfigurationDialog::OnPreviewImport() const
 			}
 		});
 
+	//Make sure all temporary object are not flags to persist
+	ClearObjectFlags(Results);
+}
+
+FReply SInterchangePipelineConfigurationDialog::OnPreviewImport() const
+{
+	auto ClearObjectFlags = [](UObject* Obj)
+		{
+			Obj->ClearFlags(RF_Standalone | RF_Public);
+			Obj->ClearInternalFlags(EInternalObjectFlags::Async);
+		};
+
+	constexpr bool bUpdateCardsFalse = false;
+	UpdatePreviewContainer(bUpdateCardsFalse);
+
+	if (!ensure(PreviewNodeContainer))
+	{
+		return FReply::Handled();
+	}
+
 	//Create and show the graph inspector UI dialog
 	TSharedRef<SWindow> Window = SNew(SWindow)
 		.ClientSize(FVector2D(800.f, 650.f))
@@ -1146,18 +1627,13 @@ FReply SInterchangePipelineConfigurationDialog::OnPreviewImport() const
 	Window->SetContent
 	(
 		SAssignNew(InterchangeGraphInspectorWindow, SInterchangeGraphInspectorWindow)
-		.InterchangeBaseNodeContainer(DuplicateBaseNodeContainer)
+		.InterchangeBaseNodeContainer(PreviewNodeContainer)
 		.bPreview(true)
 		.OwnerWindow(Window)
 	);
 
 	FSlateApplication::Get().AddModalWindow(Window, OwnerWindow.Pin(), false);
 
-	//Make sure all temporary object are not flags to persist
-	//We cannot run a gc now since the pipeline we will return are not yet hold by the AsyncHelper, so they will be garbage collect
-	ClearObjectFlags(DuplicateBaseNodeContainer);
-	ClearObjectFlags(Results);
-
 	return FReply::Handled();
 }
 
diff --git a/Engine/Plugins/Interchange/Editor/Source/Pipelines/Private/SInterchangePipelineConfigurationDialog.h b/Engine/Plugins/Interchange/Editor/Source/Pipelines/Private/SInterchangePipelineConfigurationDialog.h
--- a/Engine/Plugins/Interchange/Editor/Source/Pipelines/Private/SInterchangePipelineConfigurationDialog.h
+++ b/Engine/Plugins/Interchange/Editor/Source/Pipelines/Private/SInterchangePipelineConfigurationDialog.h
@@ -5,7 +5,9 @@
 #include "Nodes/InterchangeBaseNodeContainer.h"
 #include "InterchangePipelineBase.h"
 #include "InterchangePipelineConfigurationBase.h"
+#include "SInterchangeAssetCard.h"
 #include "Styling/SlateBrush.h"
+#include "Widgets/Layout/SScrollBox.h"
 #include "Widgets/Views/STreeView.h"
 
 class SBox;
@@ -15,6 +17,7 @@ struct FSlateBrush;
 class IDetailsView;
 class SCheckBox;
 class STextComboBox;
+class UInterchangeCardsPipeline;
 class UInterchangeTranslatorBase;
 
 
@@ -26,7 +29,7 @@ public:
 	UObject* ReimportObject = nullptr;
 	UInterchangeBaseNodeContainer* Container = nullptr;
 	UInterchangeSourceData* SourceData = nullptr;
-	bool bBasicLayout = false;
+	bool bShowEssentials = false;
 	TArray<FInterchangeConflictInfo> ConflictInfos;
 };
 
@@ -99,6 +102,7 @@ public:
 
 private:
 	TSharedRef<SBox> SpawnPipelineConfiguration();
+	TSharedRef<SBox> SpawnCardsConfiguration();
 
 	bool IsPropertyVisible(const FPropertyAndParent&) const;
 	FText GetSourceDescription() const;
@@ -119,8 +123,11 @@ private:
 
 private:
 	TWeakPtr< SWindow > OwnerWindow;
+	double OriginalMinWindowSize = 0.0;
+	double DeltaClientWindowSize = 0.0;
 	TWeakObjectPtr<UInterchangeSourceData> SourceData;
 	TWeakObjectPtr<UInterchangeBaseNodeContainer> BaseNodeContainer;
+	mutable TObjectPtr<UInterchangeBaseNodeContainer> PreviewNodeContainer = nullptr;
 	TWeakObjectPtr<UObject> ReimportObject;
 	TWeakObjectPtr<UInterchangeTranslatorBase> Translator;
 	TObjectPtr<UInterchangeTranslatorSettings> TranslatorSettings = nullptr;
@@ -145,22 +152,44 @@ private:
 
 	//
 	//////////////////////////////////////////////////////////////////////////
+	struct FFactoryNodeEnabledData
+	{
+		bool bEnable = true;
+		TObjectPtr<UClass> ObjectClass = nullptr;
+	};
+	TArray<TObjectPtr<UClass>> PipelineSupportAssetClasses;
+	TSharedPtr<SScrollBar> CardViewScrollbar;
+	TMap<UClass*, FFactoryNodeEnabledData> EnableDataPerFactoryNodeClass;
+	TSharedPtr<SInterchangeAssetCardList> CardViewList = nullptr;
+	TArray<TSharedPtr<SInterchangeAssetCard>> AssetCards;
+	void UpdatePipelineSupportedAssetClasses();
+	void UpdateEnableDataPerFactoryNodeClass();
+	void FillAssetCardsList();
+	void CreateCardsViewList();
+	void RefreshCardsViewList();
+	UInterchangeCardsPipeline* GenerateTransientCardsPipeline() const;
+
+	//Splitter management
+	double SplitAdvancedRatio = 0.6;
+	TSharedPtr<SSplitter> CardsAndAdvancedSplitter = nullptr;
 
 	ECheckBoxState IsFilteringOptions() const
 	{
 		return bFilterOptions ? ECheckBoxState::Checked : ECheckBoxState::Unchecked;
 	}
 
-	ECheckBoxState IsBasicLayoutEnabled() const
+	ECheckBoxState IsShowEssentialsEnabled() const
 	{
-		return bBasicLayout ? ECheckBoxState::Checked : ECheckBoxState::Unchecked;
+		return bShowEssentials ? ECheckBoxState::Checked : ECheckBoxState::Unchecked;
 	}
 
 	void OnFilterOptionsChanged(ECheckBoxState CheckState);
-	void OnBasicLayoutChanged(ECheckBoxState CheckState);
+	void OnShowEssentialsChanged(ECheckBoxState CheckState);
 
 	const FSlateBrush* GetImportButtonIcon() const;
 
+	void UpdatePreviewContainer(bool bUpdateCards) const;
+
 	FReply OnPreviewImport() const;
 
 	TSharedPtr<IDetailsView> PipelineConfigurationDetailsView;
@@ -172,7 +201,9 @@ private:
 	bool bImportAll = false;
 
 	bool bFilterOptions = false;
-	bool bBasicLayout = false;
+	bool bShowEssentials = false;
+	bool bShowSettings = false;
+	bool bShowCards = true;
 
 	FName CurrentStackName = NAME_None;
 	TObjectPtr<UInterchangePipelineBase> CurrentSelectedPipeline = nullptr;
diff --git a/Engine/Plugins/Interchange/Editor/Source/Pipelines/Public/InterchangeCardsPipeline.h b/Engine/Plugins/Interchange/Editor/Source/Pipelines/Public/InterchangeCardsPipeline.h
new file mode 100644
--- /dev/null
+++ b/Engine/Plugins/Interchange/Editor/Source/Pipelines/Public/InterchangeCardsPipeline.h
@@ -0,0 +1,57 @@
+// Copyright Epic Games, Inc. All Rights Reserved.
+
+#pragma once
+
+#include "Containers/Array.h"
+#include "InterchangePipelineBase.h"
+#include "UObject/Object.h"
+#include "UObject/ObjectMacros.h"
+#include "UObject/ObjectPtr.h"
+#include "UObject/UObjectGlobals.h"
+
+#include "InterchangeCardsPipeline.generated.h"
+
+class UInterchangeBaseNodeContainer;
+class UInterchangeSourceData;
+/**
+ * This pipeline is use by the interchange default UI to enable and disable factory node.
+ *
+ * @note This is an import only pipeline that is execute only when importing from the interchange default dialog.
+ */
+UCLASS()
+class INTERCHANGEEDITORPIPELINES_API UInterchangeCardsPipeline : public UInterchangePipelineBase
+{
+	GENERATED_BODY()
+
+public:
+	//////////////////////////////////////////////////////////////////////////
+	// BEGIN Pre import pipeline properties
+
+
+	// END Pre import pipeline properties
+	//////////////////////////////////////////////////////////////////////////
+
+	/*
+	 * This pipeline must never be saved into any asset import data.
+	 */
+	virtual bool SupportReimport() const override { return false; }
+
+	/*
+	 * Set all factory node class to disabled.
+	 */
+	void SetDisabledFactoryNodes(TArray<UClass*> FactoryNodeClasses);
+
+protected:
+
+	virtual void ExecutePipeline(UInterchangeBaseNodeContainer* BaseNodeContainer, const TArray<UInterchangeSourceData*>& SourceDatas, const FString& ContentBasePath) override;
+
+	virtual bool CanExecuteOnAnyThread(EInterchangePipelineTask PipelineTask) override
+	{
+		//If a Blueprint or Python class derived from this class, it will be executed on the game thread because we cannot currently execute script outside of the game thread, even if this function return true.
+		return true;
+	}
+private:
+
+	UPROPERTY()
+	TArray<TObjectPtr<UClass>> FactoryNodeClassesToDisabled;
+};
diff --git a/Engine/Plugins/Interchange/Runtime/Source/Pipelines/InterchangePipelines.Build.cs b/Engine/Plugins/Interchange/Runtime/Source/Pipelines/InterchangePipelines.Build.cs
--- a/Engine/Plugins/Interchange/Runtime/Source/Pipelines/InterchangePipelines.Build.cs
+++ b/Engine/Plugins/Interchange/Runtime/Source/Pipelines/InterchangePipelines.Build.cs
@@ -16,7 +16,7 @@ namespace UnrealBuildTool.Rules
 					"InterchangeCommon",
 					"InterchangeEngine",
 					"InterchangeFactoryNodes",
-					"InterchangeNodes"
+					"InterchangeNodes",
 				}
 			);
 
@@ -31,6 +31,7 @@ namespace UnrealBuildTool.Rules
 					"InputCore",
 					"Slate",
 					"SlateCore",
+					"LevelSequence",
 				}
 			);
 
diff --git a/Engine/Plugins/Interchange/Runtime/Source/Pipelines/Private/InterchangeGenericAnimationPipeline.cpp b/Engine/Plugins/Interchange/Runtime/Source/Pipelines/Private/InterchangeGenericAnimationPipeline.cpp
--- a/Engine/Plugins/Interchange/Runtime/Source/Pipelines/Private/InterchangeGenericAnimationPipeline.cpp
+++ b/Engine/Plugins/Interchange/Runtime/Source/Pipelines/Private/InterchangeGenericAnimationPipeline.cpp
@@ -2,6 +2,7 @@
 #include "InterchangeGenericAnimationPipeline.h"
 
 #include "Animation/AnimationSettings.h"
+#include "Animation/AnimSequence.h"
 #include "CoreMinimal.h"
 #include "Engine/SkeletalMesh.h"
 #include "Engine/StaticMesh.h"
@@ -17,6 +18,7 @@
 #include "InterchangeSkeletonFactoryNode.h"
 #include "InterchangeSkeletonHelper.h"
 #include "InterchangeSourceData.h"
+#include "LevelSequence.h"
 #include "Nodes/InterchangeBaseNode.h"
 #include "Nodes/InterchangeBaseNodeContainer.h"
 #include "Nodes/InterchangeSourceNode.h"
@@ -161,6 +163,26 @@ void UInterchangeGenericAnimationPipeline::AdjustSettingsForContext(const FInter
 #endif //WITH_EDITOR
 }
 
+#if WITH_EDITOR
+
+bool UInterchangeGenericAnimationPipeline::IsPropertyChangeNeedRefresh(const FPropertyChangedEvent& PropertyChangedEvent) const
+{
+	if (PropertyChangedEvent.GetPropertyName() == GET_MEMBER_NAME_CHECKED(UInterchangeGenericAnimationPipeline, bImportAnimations))
+	{
+		return true;
+	}
+	return Super::IsPropertyChangeNeedRefresh(PropertyChangedEvent);
+}
+
+
+void UInterchangeGenericAnimationPipeline::GetSupportAssetClasses(TArray<UClass*>& PipelineSupportAssetClasses) const
+{
+	PipelineSupportAssetClasses.Add(UAnimSequence::StaticClass());
+	PipelineSupportAssetClasses.Add(ULevelSequence::StaticClass());
+}
+
+#endif //WITH_EDITOR
+
 void UInterchangeGenericAnimationPipeline::ExecutePipeline(UInterchangeBaseNodeContainer* InBaseNodeContainer, const TArray<UInterchangeSourceData*>& InSourceDatas, const FString& ContentBasePath)
 {
 	if (!InBaseNodeContainer)
diff --git a/Engine/Plugins/Interchange/Runtime/Source/Pipelines/Private/InterchangeGenericAssetsPipeline.cpp b/Engine/Plugins/Interchange/Runtime/Source/Pipelines/Private/InterchangeGenericAssetsPipeline.cpp
--- a/Engine/Plugins/Interchange/Runtime/Source/Pipelines/Private/InterchangeGenericAssetsPipeline.cpp
+++ b/Engine/Plugins/Interchange/Runtime/Source/Pipelines/Private/InterchangeGenericAssetsPipeline.cpp
@@ -375,8 +375,19 @@ void UInterchangeGenericAssetsPipeline::FilterPropertiesFromTranslatedData(UInte
 	}
 }
 
-bool UInterchangeGenericAssetsPipeline::IsPropertyChangeNeedRefresh(const FPropertyChangedEvent& PropertyChangedEvent)
+bool UInterchangeGenericAssetsPipeline::IsPropertyChangeNeedRefresh(const FPropertyChangedEvent& PropertyChangedEvent) const
 {
+	static const TSet<FName> NeedRefreshProperties =
+	{
+		GET_MEMBER_NAME_CHECKED(UInterchangeGenericAssetsPipeline, bUseSourceNameForAsset),
+		GET_MEMBER_NAME_CHECKED(UInterchangeGenericAssetsPipeline, AssetName)
+	};
+
+	if (NeedRefreshProperties.Contains(PropertyChangedEvent.GetPropertyName()))
+	{
+		return true;
+	}
+
 	if ((CommonMeshesProperties && CommonMeshesProperties->IsPropertyChangeNeedRefresh(PropertyChangedEvent))
 		|| (CommonSkeletalMeshesAndAnimationsProperties && CommonSkeletalMeshesAndAnimationsProperties->IsPropertyChangeNeedRefresh(PropertyChangedEvent))
 		|| (MeshPipeline && MeshPipeline->IsPropertyChangeNeedRefresh(PropertyChangedEvent))
@@ -388,6 +399,22 @@ bool UInterchangeGenericAssetsPipeline::IsPropertyChangeNeedRefresh(const FPrope
 	return Super::IsPropertyChangeNeedRefresh(PropertyChangedEvent);
 }
 
+void UInterchangeGenericAssetsPipeline::GetSupportAssetClasses(TArray<UClass*>& PipelineSupportAssetClasses) const
+{
+	if (MeshPipeline)
+	{
+		MeshPipeline->GetSupportAssetClasses(PipelineSupportAssetClasses);
+	}
+	if (MaterialPipeline)
+	{
+		MaterialPipeline->GetSupportAssetClasses(PipelineSupportAssetClasses);
+	}
+	if (AnimationPipeline)
+	{
+		AnimationPipeline->GetSupportAssetClasses(PipelineSupportAssetClasses);
+	}
+}
+
 bool UInterchangeGenericAssetsPipeline::GetPropertyPossibleValues(const FName PropertyPath, TArray<FString>& PossibleValues)
 {
 	const FString PropertyPathString = PropertyPath.ToString();
diff --git a/Engine/Plugins/Interchange/Runtime/Source/Pipelines/Private/InterchangeGenericMaterialPipeline.cpp b/Engine/Plugins/Interchange/Runtime/Source/Pipelines/Private/InterchangeGenericMaterialPipeline.cpp
--- a/Engine/Plugins/Interchange/Runtime/Source/Pipelines/Private/InterchangeGenericMaterialPipeline.cpp
+++ b/Engine/Plugins/Interchange/Runtime/Source/Pipelines/Private/InterchangeGenericMaterialPipeline.cpp
@@ -697,8 +697,13 @@ void UInterchangeGenericMaterialPipeline::FilterPropertiesFromTranslatedData(UIn
 	}
 }
 
-bool UInterchangeGenericMaterialPipeline::IsPropertyChangeNeedRefresh(const FPropertyChangedEvent& PropertyChangedEvent)
+bool UInterchangeGenericMaterialPipeline::IsPropertyChangeNeedRefresh(const FPropertyChangedEvent& PropertyChangedEvent) const
 {
+	if (PropertyChangedEvent.GetPropertyName() == GET_MEMBER_NAME_CHECKED(UInterchangeGenericMaterialPipeline, bImportMaterials))
+	{
+		return true;
+	}
+
 	if (TexturePipeline && TexturePipeline->IsPropertyChangeNeedRefresh(PropertyChangedEvent))
 	{
 		return true;
@@ -706,6 +711,17 @@ bool UInterchangeGenericMaterialPipeline::IsPropertyChangeNeedRefresh(const FPro
 	return Super::IsPropertyChangeNeedRefresh(PropertyChangedEvent);
 }
 
+void UInterchangeGenericMaterialPipeline::GetSupportAssetClasses(TArray<UClass*>& PipelineSupportAssetClasses) const
+{
+	if (TexturePipeline)
+	{
+		TexturePipeline->GetSupportAssetClasses(PipelineSupportAssetClasses);
+	}
+
+	PipelineSupportAssetClasses.Add(UMaterial::StaticClass());
+	PipelineSupportAssetClasses.Add(UMaterialInstance::StaticClass());
+}
+
 #endif //WITH_EDITOR
 
 void UInterchangeGenericMaterialPipeline::ExecutePipeline(UInterchangeBaseNodeContainer* InBaseNodeContainer, const TArray<UInterchangeSourceData*>& InSourceDatas, const FString& ContentBasePath)
diff --git a/Engine/Plugins/Interchange/Runtime/Source/Pipelines/Private/InterchangeGenericMeshPipeline.cpp b/Engine/Plugins/Interchange/Runtime/Source/Pipelines/Private/InterchangeGenericMeshPipeline.cpp
--- a/Engine/Plugins/Interchange/Runtime/Source/Pipelines/Private/InterchangeGenericMeshPipeline.cpp
+++ b/Engine/Plugins/Interchange/Runtime/Source/Pipelines/Private/InterchangeGenericMeshPipeline.cpp
@@ -224,9 +224,18 @@ void UInterchangeGenericMeshPipeline::AdjustSettingsForContext(const FInterchang
 
 #if WITH_EDITOR
 
-bool UInterchangeGenericMeshPipeline::IsPropertyChangeNeedRefresh(const FPropertyChangedEvent& PropertyChangedEvent)
+bool UInterchangeGenericMeshPipeline::IsPropertyChangeNeedRefresh(const FPropertyChangedEvent& PropertyChangedEvent) const
 {
-	if (PropertyChangedEvent.GetPropertyName() == GET_MEMBER_NAME_CHECKED(UInterchangeGenericMeshPipeline, SkeletalMeshImportContentType))
+	static const TSet<FName> NeedRefreshProperties =
+	{
+		GET_MEMBER_NAME_CHECKED(UInterchangeGenericMeshPipeline, SkeletalMeshImportContentType),
+		GET_MEMBER_NAME_CHECKED(UInterchangeGenericMeshPipeline, bImportStaticMeshes),
+		GET_MEMBER_NAME_CHECKED(UInterchangeGenericMeshPipeline, bImportSkeletalMeshes),
+		GET_MEMBER_NAME_CHECKED(UInterchangeGenericMeshPipeline, bCreatePhysicsAsset),
+		GET_MEMBER_NAME_CHECKED(UInterchangeGenericMeshPipeline, bCombineStaticMeshes)
+	};
+	
+	if (NeedRefreshProperties.Contains(PropertyChangedEvent.GetPropertyName()))
 	{
 		return true;
 	}
@@ -268,6 +277,16 @@ bool UInterchangeGenericMeshPipeline::GetPropertyPossibleValues(const FName Prop
 	return Super::GetPropertyPossibleValues(PropertyPath, PossibleValues);
 }
 
+void UInterchangeGenericMeshPipeline::GetSupportAssetClasses(TArray<UClass*>& PipelineSupportAssetClasses) const
+{
+	PipelineSupportAssetClasses.Add(UStaticMesh::StaticClass());
+	PipelineSupportAssetClasses.Add(USkeletalMesh::StaticClass());
+	if (bCreatePhysicsAsset && !PhysicsAsset.IsValid())
+	{
+		PipelineSupportAssetClasses.Add(UPhysicsAsset::StaticClass());
+	}
+}
+
 #endif
 
 void UInterchangeGenericMeshPipeline::GetMeshesInformationFromTranslatedData(const UInterchangeBaseNodeContainer* InBaseNodeContainer
diff --git a/Engine/Plugins/Interchange/Runtime/Source/Pipelines/Private/InterchangeGenericTexturePipeline.cpp b/Engine/Plugins/Interchange/Runtime/Source/Pipelines/Private/InterchangeGenericTexturePipeline.cpp
--- a/Engine/Plugins/Interchange/Runtime/Source/Pipelines/Private/InterchangeGenericTexturePipeline.cpp
+++ b/Engine/Plugins/Interchange/Runtime/Source/Pipelines/Private/InterchangeGenericTexturePipeline.cpp
@@ -144,6 +144,15 @@ void UInterchangeGenericTexturePipeline::AdjustSettingsForContext(const FInterch
 
 #if WITH_EDITOR
 
+bool UInterchangeGenericTexturePipeline::IsPropertyChangeNeedRefresh(const FPropertyChangedEvent& PropertyChangedEvent) const
+{
+	if (PropertyChangedEvent.GetPropertyName() == GET_MEMBER_NAME_CHECKED(UInterchangeGenericTexturePipeline, bImportTextures))
+	{
+		return true;
+	}
+	return Super::IsPropertyChangeNeedRefresh(PropertyChangedEvent);
+}
+
 void UInterchangeGenericTexturePipeline::FilterPropertiesFromTranslatedData(UInterchangeBaseNodeContainer* InBaseNodeContainer)
 {
 	Super::FilterPropertiesFromTranslatedData(InBaseNodeContainer);
@@ -161,6 +170,11 @@ void UInterchangeGenericTexturePipeline::FilterPropertiesFromTranslatedData(UInt
 	}
 }
 
+void UInterchangeGenericTexturePipeline::GetSupportAssetClasses(TArray<UClass*>& PipelineSupportAssetClasses) const
+{
+	PipelineSupportAssetClasses.Add(UTexture::StaticClass());
+}
+
 #endif //WITH_EDITOR
 
 void UInterchangeGenericTexturePipeline::ExecutePipeline(UInterchangeBaseNodeContainer* InBaseNodeContainer, const TArray<UInterchangeSourceData*>& InSourceDatas, const FString& ContentBasePath)
diff --git a/Engine/Plugins/Interchange/Runtime/Source/Pipelines/Public/InterchangeGenericAnimationPipeline.h b/Engine/Plugins/Interchange/Runtime/Source/Pipelines/Public/InterchangeGenericAnimationPipeline.h
--- a/Engine/Plugins/Interchange/Runtime/Source/Pipelines/Public/InterchangeGenericAnimationPipeline.h
+++ b/Engine/Plugins/Interchange/Runtime/Source/Pipelines/Public/InterchangeGenericAnimationPipeline.h
@@ -121,7 +121,15 @@ public:
 	FString SourceAnimationName;
 
 	virtual void AdjustSettingsForContext(const FInterchangePipelineContextParams& ContextParams) override;
+
+#if WITH_EDITOR
+	
+	virtual bool IsPropertyChangeNeedRefresh(const FPropertyChangedEvent& PropertyChangedEvent) const override;
 	
+	virtual void GetSupportAssetClasses(TArray<UClass*>& PipelineSupportAssetClasses) const override;
+
+#endif //WITH_EDITOR
+
 protected:
 
 	virtual void ExecutePipeline(UInterchangeBaseNodeContainer* InBaseNodeContainer, const TArray<UInterchangeSourceData*>& InSourceDatas, const FString& ContentBasePath) override;
diff --git a/Engine/Plugins/Interchange/Runtime/Source/Pipelines/Public/InterchangeGenericAssetsPipeline.h b/Engine/Plugins/Interchange/Runtime/Source/Pipelines/Public/InterchangeGenericAssetsPipeline.h
--- a/Engine/Plugins/Interchange/Runtime/Source/Pipelines/Public/InterchangeGenericAssetsPipeline.h
+++ b/Engine/Plugins/Interchange/Runtime/Source/Pipelines/Public/InterchangeGenericAssetsPipeline.h
@@ -102,7 +102,9 @@ public:
 #if WITH_EDITOR
 	virtual void FilterPropertiesFromTranslatedData(UInterchangeBaseNodeContainer* InBaseNodeContainer) override;
 
-	virtual bool IsPropertyChangeNeedRefresh(const FPropertyChangedEvent& PropertyChangedEvent) override;
+	virtual bool IsPropertyChangeNeedRefresh(const FPropertyChangedEvent& PropertyChangedEvent) const override;
+
+	virtual void GetSupportAssetClasses(TArray<UClass*>& PipelineSupportAssetClasses) const override;
 #endif //WITH_EDITOR
 
 	virtual TArray<FInterchangeConflictInfo> GetConflictInfos(UObject* ReimportObject, UInterchangeBaseNodeContainer* InBaseNodeContainer, UInterchangeSourceData* SourceData) override;
diff --git a/Engine/Plugins/Interchange/Runtime/Source/Pipelines/Public/InterchangeGenericAssetsPipelineSharedSettings.h b/Engine/Plugins/Interchange/Runtime/Source/Pipelines/Public/InterchangeGenericAssetsPipelineSharedSettings.h
--- a/Engine/Plugins/Interchange/Runtime/Source/Pipelines/Public/InterchangeGenericAssetsPipelineSharedSettings.h
+++ b/Engine/Plugins/Interchange/Runtime/Source/Pipelines/Public/InterchangeGenericAssetsPipelineSharedSettings.h
@@ -130,13 +130,15 @@ public:
 	UPROPERTY(EditAnywhere, BlueprintReadWrite, Category = "Common Meshes", meta = (SubCategory = "Build"))
 	bool bRemoveDegenerates = false;
 #if WITH_EDITOR
-	virtual bool IsPropertyChangeNeedRefresh(const FPropertyChangedEvent& PropertyChangedEvent) override
+	virtual bool IsPropertyChangeNeedRefresh(const FPropertyChangedEvent& PropertyChangedEvent) const override
 	{
-		if (PropertyChangedEvent.Property->GetName() == GET_MEMBER_NAME_CHECKED(UInterchangeGenericCommonMeshesProperties, ForceAllMeshAsType))
+		static const TSet<FName> NeedRefreshProperties =
 		{
-			return true;
-		}
-		if (PropertyChangedEvent.Property->GetName() == GET_MEMBER_NAME_CHECKED(UInterchangeGenericCommonMeshesProperties, bAutoDetectMeshType))
+			GET_MEMBER_NAME_CHECKED(UInterchangeGenericCommonMeshesProperties, ForceAllMeshAsType),
+			GET_MEMBER_NAME_CHECKED(UInterchangeGenericCommonMeshesProperties, bAutoDetectMeshType)
+		};
+
+		if (NeedRefreshProperties.Contains(PropertyChangedEvent.GetPropertyName()))
 		{
 			return true;
 		}
@@ -179,24 +181,21 @@ public:
 	UPROPERTY(EditAnywhere, BlueprintReadWrite, Category = "Static Meshes")
 	bool bConvertStaticsWithMorphTargetsToSkeletals = false;
 #if WITH_EDITOR
-	virtual bool IsPropertyChangeNeedRefresh(const FPropertyChangedEvent& PropertyChangedEvent) override
+	virtual bool IsPropertyChangeNeedRefresh(const FPropertyChangedEvent& PropertyChangedEvent) const override
 	{
-		if (PropertyChangedEvent.Property->GetName() == GET_MEMBER_NAME_CHECKED(UInterchangeGenericCommonSkeletalMeshesAndAnimationsProperties, bConvertStaticsWithMorphTargetsToSkeletals))
+		static const TSet<FName> NeedRefreshProperties =
 		{
-			return true;
-		}
-		else if (PropertyChangedEvent.Property->GetName() == GET_MEMBER_NAME_CHECKED(UInterchangeGenericCommonSkeletalMeshesAndAnimationsProperties, bImportMeshesInBoneHierarchy))
-		{
-			return true;
-		}
-		else if (PropertyChangedEvent.Property->GetName() == GET_MEMBER_NAME_CHECKED(UInterchangeGenericCommonSkeletalMeshesAndAnimationsProperties, Skeleton))
-		{
-			return true;
-		}
-		else if (PropertyChangedEvent.Property->GetName() == GET_MEMBER_NAME_CHECKED(UInterchangeGenericCommonSkeletalMeshesAndAnimationsProperties, bImportOnlyAnimations))
+			GET_MEMBER_NAME_CHECKED(UInterchangeGenericCommonSkeletalMeshesAndAnimationsProperties, bConvertStaticsWithMorphTargetsToSkeletals),
+			GET_MEMBER_NAME_CHECKED(UInterchangeGenericCommonSkeletalMeshesAndAnimationsProperties, bImportMeshesInBoneHierarchy),
+			GET_MEMBER_NAME_CHECKED(UInterchangeGenericCommonSkeletalMeshesAndAnimationsProperties, Skeleton),
+			GET_MEMBER_NAME_CHECKED(UInterchangeGenericCommonSkeletalMeshesAndAnimationsProperties, bImportOnlyAnimations)
+		};
+
+		if (NeedRefreshProperties.Contains(PropertyChangedEvent.GetPropertyName()))
 		{
 			return true;
 		}
+
 		return Super::IsPropertyChangeNeedRefresh(PropertyChangedEvent);
 	}
 #endif //WITH_EDITOR
diff --git a/Engine/Plugins/Interchange/Runtime/Source/Pipelines/Public/InterchangeGenericMaterialPipeline.h b/Engine/Plugins/Interchange/Runtime/Source/Pipelines/Public/InterchangeGenericMaterialPipeline.h
--- a/Engine/Plugins/Interchange/Runtime/Source/Pipelines/Public/InterchangeGenericMaterialPipeline.h
+++ b/Engine/Plugins/Interchange/Runtime/Source/Pipelines/Public/InterchangeGenericMaterialPipeline.h
@@ -104,7 +104,9 @@ public:
 	virtual void AdjustSettingsForContext(const FInterchangePipelineContextParams& ContextParams) override;
 #if WITH_EDITOR
 	virtual void FilterPropertiesFromTranslatedData(UInterchangeBaseNodeContainer* InBaseNodeContainer) override;
-	virtual bool IsPropertyChangeNeedRefresh(const FPropertyChangedEvent& PropertyChangedEvent) override;
+	virtual bool IsPropertyChangeNeedRefresh(const FPropertyChangedEvent& PropertyChangedEvent) const override;
+
+	virtual void GetSupportAssetClasses(TArray<UClass*>& PipelineSupportAssetClasses) const override;
 #endif //WITH_EDITOR
 
 protected:
diff --git a/Engine/Plugins/Interchange/Runtime/Source/Pipelines/Public/InterchangeGenericMeshPipeline.h b/Engine/Plugins/Interchange/Runtime/Source/Pipelines/Public/InterchangeGenericMeshPipeline.h
--- a/Engine/Plugins/Interchange/Runtime/Source/Pipelines/Public/InterchangeGenericMeshPipeline.h
+++ b/Engine/Plugins/Interchange/Runtime/Source/Pipelines/Public/InterchangeGenericMeshPipeline.h
@@ -230,8 +230,10 @@ public:
 	virtual void PreDialogCleanup(const FName PipelineStackName) override;
 
 #if WITH_EDITOR
-	virtual bool IsPropertyChangeNeedRefresh(const FPropertyChangedEvent& PropertyChangedEvent) override;
+	virtual bool IsPropertyChangeNeedRefresh(const FPropertyChangedEvent& PropertyChangedEvent) const override;
 	virtual bool GetPropertyPossibleValues(const FName PropertyPath, TArray<FString>& PossibleValues) override;
+
+	virtual void GetSupportAssetClasses(TArray<UClass*>& PipelineSupportAssetClasses) const override;
 #endif
 
 	static UInterchangePipelineMeshesUtilities* CreateMeshPipelineUtilities(UInterchangeBaseNodeContainer* InBaseNodeContainer
diff --git a/Engine/Plugins/Interchange/Runtime/Source/Pipelines/Public/InterchangeGenericTexturePipeline.h b/Engine/Plugins/Interchange/Runtime/Source/Pipelines/Public/InterchangeGenericTexturePipeline.h
--- a/Engine/Plugins/Interchange/Runtime/Source/Pipelines/Public/InterchangeGenericTexturePipeline.h
+++ b/Engine/Plugins/Interchange/Runtime/Source/Pipelines/Public/InterchangeGenericTexturePipeline.h
@@ -74,8 +74,12 @@ public:
 
 #if WITH_EDITOR
 
+	virtual bool IsPropertyChangeNeedRefresh(const FPropertyChangedEvent& PropertyChangedEvent) const override;
+
 	virtual void FilterPropertiesFromTranslatedData(UInterchangeBaseNodeContainer* InBaseNodeContainer) override;
 
+	virtual void GetSupportAssetClasses(TArray<UClass*>& PipelineSupportAssetClasses) const override;
+
 #endif //WITH_EDITOR
 
 protected:
diff --git a/Engine/Source/Runtime/Interchange/Core/Private/InterchangePipelineBase.cpp b/Engine/Source/Runtime/Interchange/Core/Private/InterchangePipelineBase.cpp
--- a/Engine/Source/Runtime/Interchange/Core/Private/InterchangePipelineBase.cpp
+++ b/Engine/Source/Runtime/Interchange/Core/Private/InterchangePipelineBase.cpp
@@ -69,7 +69,7 @@ void UInterchangePipelineBase::TransferAdjustSettings(UInterchangePipelineBase*
 	CachePropertiesStates = SourcePipeline->CachePropertiesStates;
 	bAllowPropertyStatesEdition = SourcePipeline->bAllowPropertyStatesEdition;
 	bIsReimportContext = SourcePipeline->bIsReimportContext;
-	bIsBasicLayout = SourcePipeline->bIsBasicLayout;
+	bIsShowEssentials = SourcePipeline->bIsShowEssentials;
 }
 
 const FInterchangePipelinePropertyStates* UInterchangePipelineBase::GetPropertyStates(const FName PropertyPath) const
diff --git a/Engine/Source/Runtime/Interchange/Core/Public/InterchangePipelineBase.h b/Engine/Source/Runtime/Interchange/Core/Public/InterchangePipelineBase.h
--- a/Engine/Source/Runtime/Interchange/Core/Public/InterchangePipelineBase.h
+++ b/Engine/Source/Runtime/Interchange/Core/Public/InterchangePipelineBase.h
@@ -106,18 +106,18 @@ struct FInterchangePipelinePropertyStates
 
 
 	/** Return true if the property is visible for the specified context. */
-	bool IsPropertyVisibleInBasicLayout() const
+	bool IsPropertyVisibleInShowEssentials() const
 	{
 		return BasicLayoutStates.bVisible;
 	}
 
 	/** Return true if the property is visible for the specified context. */
-	bool IsPropertyVisible(const bool bIsReimportContext, const bool bIsBasicLayout) const
+	bool IsPropertyVisible(const bool bIsReimportContext, const bool bIsShowEssentials) const
 	{
 		bool bVisible = bIsReimportContext ? ReimportStates.bVisible : ImportStates.bVisible;
 		if (bVisible)
 		{
-			bVisible = bIsBasicLayout ? BasicLayoutStates.bVisible : true;
+			bVisible = bIsShowEssentials ? BasicLayoutStates.bVisible : true;
 		}
 		return bVisible;
 	}
@@ -132,7 +132,7 @@ struct FInterchangePipelinePropertyStates
 		ReimportStates.bVisible = bVisibleValue;
 	}
 
-	void SetPropertyBasicLayoutVisibility(const bool bVisibleValue)
+	void SetPropertyShowEssentialsVisibility(const bool bVisibleValue)
 	{
 		BasicLayoutStates.bVisible = bVisibleValue;
 	}
@@ -334,9 +334,9 @@ public:
 	/** Transfer the source pipeline adjust settings to this pipeline. */
 	INTERCHANGECORE_API void TransferAdjustSettings(UInterchangePipelineBase* SourcePipeline);
 
-	INTERCHANGECORE_API void SetBasicLayoutMode(bool bBasicLayoutModeValue)
+	INTERCHANGECORE_API void SetShowEssentialsMode(bool bShowEssentialsModeValue)
 	{
-		bIsBasicLayout = bBasicLayoutModeValue;
+		bIsShowEssentials = bShowEssentialsModeValue;
 	}
 
 	/*
@@ -350,7 +350,7 @@ public:
 	/*
 	 * Return true if the pipeline was created for a re-import or an override pipelines stack
 	 */
-	INTERCHANGECORE_API bool IsFromReimportOrOverride()
+	INTERCHANGECORE_API bool IsFromReimportOrOverride() const
 	{
 		return bFromReimportOrOverride;
 	}
@@ -383,11 +383,17 @@ public:
 	 * The import dialog will call this function when the user changes a specific property. Return true if the pipeline UI should be refreshed.
 	 * A refresh will call the FilterPropertiesFromTranslatedData function.
 	 */
-	INTERCHANGECORE_API virtual bool IsPropertyChangeNeedRefresh(const FPropertyChangedEvent& PropertyChangedEvent)
+	INTERCHANGECORE_API virtual bool IsPropertyChangeNeedRefresh(const FPropertyChangedEvent& PropertyChangedEvent) const
 	{
 		return false;
 	}
 
+	/** Fill the list of all asset this pipeline can create */
+	INTERCHANGECORE_API virtual void GetSupportAssetClasses(TArray<UClass*>& PipelineSupportAssetClasses) const
+	{
+		return;
+	}
+
 #endif //WITH_EDITOR
 
 	/**
@@ -459,7 +465,7 @@ public:
 
 	bool CanEditPropertiesStates() { return bAllowPropertyStatesEdition; }
 	bool IsReimportContext() { return bIsReimportContext; }
-	bool IsBasicLayout() { return bIsBasicLayout; }
+	bool IsShowEssentials() { return bIsShowEssentials; }
 
 #if WITH_EDITOR
 	/*
@@ -579,12 +585,12 @@ protected:
 	bool bIsReimportContext = false;
 
 	/**
-	 * If true, this pipeline instance is use for basic layout.
+	 * If true, this pipeline instance is use for essentials settings layout.
 	 * If false, this pipeline instance is use for normal layout.
 	 *
 	 * Note: This layout must be set by the owner instancing this pipeline. This layout will be use to hide or not some properties.
 	 */
-	bool bIsBasicLayout = false;
+	bool bIsShowEssentials = false;
 
 	/*
 	 * If true, this pipeline was create to re-import an asset or override the project settings pipelines.
