Changed logic to re-translate a file only when the translator's settings have been modified and the settings dialog is closed
Rationale: Modifying translator's settings with numerous parameters for a large file can hinder productivity when the large file is translated for each property change. Additionally, since the translator's settings dialog is modal, there is no real value for the user.

#rnx

diff --git a/Engine/Plugins/Interchange/Editor/Source/Pipelines/Private/SInterchangePipelineConfigurationDialog.cpp b/Engine/Plugins/Interchange/Editor/Source/Pipelines/Private/SInterchangePipelineConfigurationDialog.cpp
--- a/Engine/Plugins/Interchange/Editor/Source/Pipelines/Private/SInterchangePipelineConfigurationDialog.cpp
+++ b/Engine/Plugins/Interchange/Editor/Source/Pipelines/Private/SInterchangePipelineConfigurationDialog.cpp
@@ -234,7 +234,12 @@ FReply SInterchangePipelineConfigurationDialog::OnEditTranslatorSettings()
 	DetailsViewArgs.bAllowMultipleTopLevelObjects = true;
 	DetailsViewArgs.NameAreaSettings = FDetailsViewArgs::HideNameArea;
 	TSharedRef<IDetailsView> TranslatorSettingsDetailsView = PropertyEditorModule.CreateDetailView(DetailsViewArgs);
-	TranslatorSettingsDetailsView->OnFinishedChangingProperties().AddRaw(this, &SInterchangePipelineConfigurationDialog::OnFinishedChangingProperties);
+
+	bool bIsSettingsModified = false;
+	TranslatorSettingsDetailsView->OnFinishedChangingProperties().AddLambda([&bIsSettingsModified](const FPropertyChangedEvent& PropertyChangedEvent)
+		{
+			bIsSettingsModified = true;
+		});
 	TranslatorSettingsDetailsView->SetObject(TranslatorSettings);
 
 	TSharedRef<SCustomDialog> OptionsDialog =
@@ -273,34 +278,30 @@ FReply SInterchangePipelineConfigurationDialog::OnEditTranslatorSettings()
 		});
 	OptionsDialog->ShowModal();
 
-	return FReply::Handled();
-}
-
-void SInterchangePipelineConfigurationDialog::OnFinishedChangingProperties(const FPropertyChangedEvent& PropertyChangedEvent)
-{
-	if (!Translator.IsValid() || !TranslatorSettings)
+	if (bIsSettingsModified && Translator.IsValid() && TranslatorSettings)
 	{
-		return;
-	}
-	if (UClass* TranslatorSettingsClass = TranslatorSettings->GetClass())
-	{
-		//Save the config locally before the translation.
-		TranslatorSettings->SaveSettings();
+		if (UClass* TranslatorSettingsClass = TranslatorSettings->GetClass())
+		{
+			//Save the config locally before the translation.
+			TranslatorSettings->SaveSettings();
 
-		//Need to Translate the source data
-		FScopedSlowTask Progress(2.f, NSLOCTEXT("SInterchangePipelineConfigurationDialog", "TranslatingSourceFile...", "Translating source file..."));
-		Progress.MakeDialog();
-		Progress.EnterProgressFrame(1.f);
-		//Reset the container
-		BaseNodeContainer->Reset();
+			//Need to Translate the source data
+			FScopedSlowTask Progress(2.f, NSLOCTEXT("SInterchangePipelineConfigurationDialog", "TranslatingSourceFile...", "Translating source file..."));
+			Progress.MakeDialog();
+			Progress.EnterProgressFrame(1.f);
+			//Reset the container
+			BaseNodeContainer->Reset();
 
-		Translator->Translate(*BaseNodeContainer.Get());
+			Translator->Translate(*BaseNodeContainer.Get());
 
-		//Refresh the dialog
-		RefreshStack(false);
+			//Refresh the dialog
+			RefreshStack(false);
 
-		Progress.EnterProgressFrame(1.f);
+			Progress.EnterProgressFrame(1.f);
+		}
 	}
+
+	return FReply::Handled();
 }
 
 TSharedRef<SBox> SInterchangePipelineConfigurationDialog::SpawnPipelineConfiguration()
diff --git a/Engine/Plugins/Interchange/Editor/Source/Pipelines/Private/SInterchangePipelineConfigurationDialog.h b/Engine/Plugins/Interchange/Editor/Source/Pipelines/Private/SInterchangePipelineConfigurationDialog.h
--- a/Engine/Plugins/Interchange/Editor/Source/Pipelines/Private/SInterchangePipelineConfigurationDialog.h
+++ b/Engine/Plugins/Interchange/Editor/Source/Pipelines/Private/SInterchangePipelineConfigurationDialog.h
@@ -119,7 +119,6 @@ private:
 
 	void SetEditPipeline(FInterchangePipelineItemType* PipelineItemToEdit);
 	FReply OnEditTranslatorSettings();
-	void OnFinishedChangingProperties(const FPropertyChangedEvent& PropertyChangedEvent);
 
 private:
 	TWeakPtr< SWindow > OwnerWindow;
