Interchange:
- Rename Interchange FBX/glTF/MaterialX into "Settings"
- Fix MaterialX settings not searchable in the project settings.

diff --git a/Engine/Plugins/Interchange/Editor/Source/Pipelines/Private/InterchangeMaterialXPipelineCustomizations.cpp b/Engine/Plugins/Interchange/Editor/Source/Pipelines/Private/InterchangeMaterialXPipelineCustomizations.cpp
--- a/Engine/Plugins/Interchange/Editor/Source/Pipelines/Private/InterchangeMaterialXPipelineCustomizations.cpp
+++ b/Engine/Plugins/Interchange/Editor/Source/Pipelines/Private/InterchangeMaterialXPipelineCustomizations.cpp
@@ -52,7 +52,7 @@ void FInterchangeMaterialXPipelineCustomization::CustomizeDetails(IDetailLayoutB
 		SNew(STextBlock)
 		.Font(IDetailLayoutBuilder::GetDetailFont())
 		.Text(NSLOCTEXT("InterchangeMaterialXPipelineCustomization::Message", "MaterialXPredefinedMaterialSubstitutionMessage", "See 'Project Settings > Engine > Interchange MaterialX' to edit settings."))
-		.AutoWrapText(true)
+		.AutoWrapText(false)
 	];
 }
 
@@ -97,7 +97,7 @@ void FInterchangeMaterialXPipelineSettingsCustomization::CustomizeDetails(IDetai
 			FText EnumString;
 			KeyPropertyHandle->GetValueAsDisplayText(EnumString);
 
-			FDetailWidgetRow& DetailWidgetRow = MaterialXPredefinedCategory.AddCustomRow(DisplayName);
+			FDetailWidgetRow& DetailWidgetRow = MaterialXPredefinedCategory.AddCustomRow(FText::FromStringView(TEXT("MaterialX")));
 			DetailWidgetRow
 			.NameContent()
 			[
diff --git a/Engine/Plugins/Interchange/Runtime/Source/Parsers/Fbx/Public/InterchangeFbxSettings.h b/Engine/Plugins/Interchange/Runtime/Source/Parsers/Fbx/Public/InterchangeFbxSettings.h
--- a/Engine/Plugins/Interchange/Runtime/Source/Parsers/Fbx/Public/InterchangeFbxSettings.h
+++ b/Engine/Plugins/Interchange/Runtime/Source/Parsers/Fbx/Public/InterchangeFbxSettings.h
@@ -12,7 +12,7 @@
 
 #include "InterchangeFbxSettings.generated.h"
 
-UCLASS(config = Interchange, meta = (DisplayName = "Interchange FBX"))
+UCLASS(config = Interchange, meta = (DisplayName = "FBX Settings"))
 class INTERCHANGEFBXPARSER_API UInterchangeFbxSettings : public UDeveloperSettings
 {
 	GENERATED_BODY()
diff --git a/Engine/Plugins/Interchange/Runtime/Source/Pipelines/Public/InterchangeMaterialXPipeline.h b/Engine/Plugins/Interchange/Runtime/Source/Pipelines/Public/InterchangeMaterialXPipeline.h
--- a/Engine/Plugins/Interchange/Runtime/Source/Pipelines/Public/InterchangeMaterialXPipeline.h
+++ b/Engine/Plugins/Interchange/Runtime/Source/Pipelines/Public/InterchangeMaterialXPipeline.h
@@ -21,7 +21,7 @@ uint32 INTERCHANGEPIPELINES_API GetTypeHash(EInterchangeMaterialXSettings Key);
 
 bool INTERCHANGEPIPELINES_API operator==(EInterchangeMaterialXSettings Lhs, EInterchangeMaterialXSettings Rhs);
 
-UCLASS(config = Interchange, meta = (DisplayName = "Interchange MaterialX", ToolTip = "Interchange MaterialX Pipeline Settings"))
+UCLASS(config = Interchange, meta = (DisplayName = "MaterialX Settings", ToolTip = "Interchange MaterialX Pipeline Settings"))
 class INTERCHANGEPIPELINES_API UMaterialXPipelineSettings : public UDeveloperSettings
 {
 	GENERATED_BODY()
diff --git a/Engine/Plugins/Interchange/Runtime/Source/Pipelines/Public/InterchangeglTFPipeline.h b/Engine/Plugins/Interchange/Runtime/Source/Pipelines/Public/InterchangeglTFPipeline.h
--- a/Engine/Plugins/Interchange/Runtime/Source/Pipelines/Public/InterchangeglTFPipeline.h
+++ b/Engine/Plugins/Interchange/Runtime/Source/Pipelines/Public/InterchangeglTFPipeline.h
@@ -16,7 +16,7 @@ class UInterchangeBaseNodeContainer;
 class UInterchangeMaterialInstanceFactoryNode;
 class UInterchangeShaderGraphNode;
 
-UCLASS(config = Interchange, meta = (DisplayName = "Interchange glTF", ToolTip = "Interchange settings for glTF conversions."))
+UCLASS(config = Interchange, meta = (DisplayName = "glTF Settings", ToolTip = "Interchange settings for glTF conversions."))
 class INTERCHANGEPIPELINES_API UGLTFPipelineSettings : public UDeveloperSettings
 {
 	GENERATED_BODY()
