Interchange Mesh Import: Support for custom LOD Screen Sizes

- Added an array property to the pipeline settings.
- Added Array Attribute to UInterchangeStaticMeshFactoryNode.
- This is array is loaded with the data from the Import Dialog and is later used to LOD Screen Sizes on the static mesh.

diff --git a/Engine/Plugins/Interchange/Editor/Source/InterchangeEditor/Private/InterchangeFbxAssetImportDataConverter.cpp b/Engine/Plugins/Interchange/Editor/Source/InterchangeEditor/Private/InterchangeFbxAssetImportDataConverter.cpp
--- a/Engine/Plugins/Interchange/Editor/Source/InterchangeEditor/Private/InterchangeFbxAssetImportDataConverter.cpp
+++ b/Engine/Plugins/Interchange/Editor/Source/InterchangeEditor/Private/InterchangeFbxAssetImportDataConverter.cpp
@@ -8,6 +8,7 @@
 #include "Engine/SkeletalMesh.h"
 #include "Engine/SkinnedAssetCommon.h"
 #include "Engine/StaticMesh.h"
+#include "StaticMeshResources.h"
 #include "Factories/FbxAnimSequenceImportData.h"
 #include "Factories/FbxAssetImportData.h"
 #include "Factories/FbxImportUI.h"
@@ -278,6 +279,24 @@ namespace UE::Interchange::Private
 		}
 	}
 
+	void FillInterchangeGenericAssetsPipelineFromFbxStaticMesh(UInterchangeGenericAssetsPipeline* GenericAssetPipeline, const UStaticMesh* StaticMesh)
+	{
+		if (!StaticMesh || !GenericAssetPipeline)
+		{
+			return;
+		}
+
+		GenericAssetPipeline->MeshPipeline->bAutoComputeLODScreenSizes = StaticMesh->bAutoComputeLODScreenSize;
+		if (const FStaticMeshRenderData* RenderData = StaticMesh->GetRenderData())
+		{
+			GenericAssetPipeline->MeshPipeline->LODScreenSizes.Empty();
+			for (int32 LodIndex = 0; LodIndex < MAX_STATIC_MESH_LODS; ++LodIndex)
+			{
+				GenericAssetPipeline->MeshPipeline->LODScreenSizes.Add(RenderData->ScreenSize[LodIndex].Default);
+			}
+		}
+	}
+
 	void FillInterchangeGenericAssetsPipelineFromFbxStaticMeshImportData(UInterchangeGenericAssetsPipeline* GenericAssetPipeline
 		, const UFbxStaticMeshImportData* StaticMeshImportData
 		, bool bFillBaseClass = true)
@@ -588,6 +607,20 @@ namespace UE::Interchange::Private
 		FillFbxAnimSequenceImportData(GenericAssetPipeline, DestinationData->AnimSequenceImportData);
 
 		DestinationData->bOverrideFullName = GenericAssetPipeline->bUseSourceNameForAsset;
+		
+		// LOD Screen Sizes
+		{
+			UInterchangeGenericMeshPipeline* MeshPipeline = GenericAssetPipeline->MeshPipeline;
+			DestinationData->bAutoComputeLodDistances = MeshPipeline->bAutoComputeLODScreenSizes;
+			DestinationData->LodDistance0 = MeshPipeline->LODScreenSizes.IsValidIndex(0) ? MeshPipeline->LODScreenSizes[0] : 0.0f;
+			DestinationData->LodDistance1 = MeshPipeline->LODScreenSizes.IsValidIndex(1) ? MeshPipeline->LODScreenSizes[1] : 0.0f;
+			DestinationData->LodDistance2 = MeshPipeline->LODScreenSizes.IsValidIndex(2) ? MeshPipeline->LODScreenSizes[2] : 0.0f;
+			DestinationData->LodDistance3 = MeshPipeline->LODScreenSizes.IsValidIndex(3) ? MeshPipeline->LODScreenSizes[3] : 0.0f;
+			DestinationData->LodDistance4 = MeshPipeline->LODScreenSizes.IsValidIndex(4) ? MeshPipeline->LODScreenSizes[4] : 0.0f;
+			DestinationData->LodDistance5 = MeshPipeline->LODScreenSizes.IsValidIndex(5) ? MeshPipeline->LODScreenSizes[5] : 0.0f;
+			DestinationData->LodDistance6 = MeshPipeline->LODScreenSizes.IsValidIndex(6) ? MeshPipeline->LODScreenSizes[6] : 0.0f;
+			DestinationData->LodDistance7 = MeshPipeline->LODScreenSizes.IsValidIndex(7) ? MeshPipeline->LODScreenSizes[7] : 0.0f;
+		}
 
 		//Material Options
 		DestinationData->bImportMaterials = GenericAssetPipeline->MaterialPipeline->bImportMaterials;
@@ -696,6 +729,7 @@ namespace UE::Interchange::Private
 
 			GenericAssetPipeline->MeshPipeline->CommonMeshesProperties->ForceAllMeshAsType = EInterchangeForceMeshType::IFMT_StaticMesh;
 			check(Obj->IsA<UStaticMesh>());
+			FillInterchangeGenericAssetsPipelineFromFbxStaticMesh(GenericAssetPipeline, Cast<UStaticMesh>(Obj));
 			FillInterchangeGenericAssetsPipelineFromFbxStaticMeshImportData(GenericAssetPipeline
 				, LegacyStaticMeshImportData);
 		}
diff --git a/Engine/Plugins/Interchange/Runtime/Source/FactoryNodes/Private/InterchangeStaticMeshFactoryNode.cpp b/Engine/Plugins/Interchange/Runtime/Source/FactoryNodes/Private/InterchangeStaticMeshFactoryNode.cpp
--- a/Engine/Plugins/Interchange/Runtime/Source/FactoryNodes/Private/InterchangeStaticMeshFactoryNode.cpp
+++ b/Engine/Plugins/Interchange/Runtime/Source/FactoryNodes/Private/InterchangeStaticMeshFactoryNode.cpp
@@ -47,6 +47,12 @@ namespace UE
 {
 	namespace Interchange
 	{
+		const FAttributeKey& FStaticMeshNodeStaticData::GetLODScreenSizeBaseKey()
+		{
+			static FAttributeKey LODScreenSize_BaseKey(TEXT("__LODScreenSize__"));
+			return LODScreenSize_BaseKey;
+		}
+
 		const FAttributeKey& FStaticMeshNodeStaticData::GetSocketUidsBaseKey()
 		{
 			static FAttributeKey SocketUids_BaseKey = FAttributeKey(TEXT("SocketUids"));
@@ -61,6 +67,7 @@ UInterchangeStaticMeshFactoryNode::UInterchangeStaticMeshFactoryNode()
 #if WITH_ENGINE
 	AssetClass = nullptr;
 #endif
+	LODScreenSizes.Initialize(Attributes.ToSharedRef(), UE::Interchange::FStaticMeshNodeStaticData::GetLODScreenSizeBaseKey().ToString());
 	SocketUids.Initialize(Attributes, UE::Interchange::FStaticMeshNodeStaticData::GetSocketUidsBaseKey().ToString());
 }
 
@@ -132,6 +139,37 @@ FString UInterchangeStaticMeshFactoryNode::GetAttributeCategory(const UE::Interc
 
 #endif //WITH_EDITOR
 
+bool UInterchangeStaticMeshFactoryNode::GetCustomAutoComputeLODScreenSizes(bool& AttributeValue) const
+{
+	IMPLEMENT_NODE_ATTRIBUTE_GETTER(AutoComputeLODScreenSizes, bool)
+}
+
+bool UInterchangeStaticMeshFactoryNode::SetCustomAutoComputeLODScreenSizes(const bool& AttributeValue)
+{
+	IMPLEMENT_NODE_ATTRIBUTE_SETTER_NODELEGATE(AutoComputeLODScreenSizes, bool)
+}
+
+int32 UInterchangeStaticMeshFactoryNode::GetLODScreenSizeCount() const
+{
+	return LODScreenSizes.GetCount();
+}
+
+void UInterchangeStaticMeshFactoryNode::GetLODScreenSizes(TArray<float>& OutLODScreenSizes) const
+{
+	LODScreenSizes.GetItems(OutLODScreenSizes);
+}
+
+bool UInterchangeStaticMeshFactoryNode::SetLODScreenSizes(const TArray<float>& InLODScreenSizes)
+{
+	LODScreenSizes.RemoveAllItems();
+	for (const float& ScreenSize: InLODScreenSizes)
+	{
+		LODScreenSizes.AddItem(ScreenSize);
+	}
+
+	return true;
+}
+
 bool UInterchangeStaticMeshFactoryNode::GetCustomBuildNanite(bool& AttributeValue) const
 {
 	IMPLEMENT_NODE_ATTRIBUTE_GETTER(BuildNanite, bool)
diff --git a/Engine/Plugins/Interchange/Runtime/Source/FactoryNodes/Public/InterchangeStaticMeshFactoryNode.h b/Engine/Plugins/Interchange/Runtime/Source/FactoryNodes/Public/InterchangeStaticMeshFactoryNode.h
--- a/Engine/Plugins/Interchange/Runtime/Source/FactoryNodes/Public/InterchangeStaticMeshFactoryNode.h
+++ b/Engine/Plugins/Interchange/Runtime/Source/FactoryNodes/Public/InterchangeStaticMeshFactoryNode.h
@@ -19,6 +19,7 @@ namespace UE
 	{
 		struct INTERCHANGEFACTORYNODES_API FStaticMeshNodeStaticData : public FBaseNodeStaticData
 		{
+			static const FAttributeKey& GetLODScreenSizeBaseKey();
 			static const FAttributeKey& GetSocketUidsBaseKey();
 		};
 	} // namespace Interchange
@@ -59,6 +60,26 @@ public:
 #endif //WITH_EDITOR
 
 public:
+	/** Get whether the static mesh factory should auto compute LOD Screen Sizes. Return false if the attribute was not set. */
+	UFUNCTION(BlueprintCallable, Category = "Interchange | Node | StaticMesh")
+	bool GetCustomAutoComputeLODScreenSizes(bool& AttributeValue) const;
+
+	/** Set whether the static mesh factory should auto compute LOD Screen Sizes. Return false if the attribute was not set. */
+	UFUNCTION(BlueprintCallable, Category = "Interchange | Node | StaticMesh")
+	bool SetCustomAutoComputeLODScreenSizes(const bool& AttributeValue);
+
+	/** Returns the number of LOD Screen Sizes the static mesh has.*/
+	UFUNCTION(BlueprintCallable, Category = "Interchange | Node | StaticMesh")
+	int32 GetLODScreenSizeCount() const;
+
+	/** Returns All the LOD Screen Sizes set for the static mesh.*/
+	UFUNCTION(BlueprintCallable, Category = "Interchange | Node | StaticMesh")
+	void GetLODScreenSizes(TArray<float>& OutLODScreenSizes) const;
+
+	/** Sets the LOD Screen Sizes for the static mesh.*/
+	UFUNCTION(BlueprintCallable, Category = "Interchange | Node | StaticMesh")
+	bool SetLODScreenSizes(const TArray<float>& InLODScreenSizes);
+
 	/** Get whether the static mesh factory should set the Nanite build setting. Return false if the attribute was not set. */
 	UFUNCTION(BlueprintCallable, Category = "Interchange | Node | StaticMesh")
 	bool GetCustomBuildNanite(bool& AttributeValue) const;
@@ -209,7 +230,10 @@ private:
 	const UE::Interchange::FAttributeKey Macro_CustomDistanceFieldReplacementMeshKey = UE::Interchange::FAttributeKey(TEXT("DistanceFieldReplacementMesh"));
 	const UE::Interchange::FAttributeKey Macro_CustomMaxLumenMeshCardsKey = UE::Interchange::FAttributeKey(TEXT("MaxLumenMeshCards"));
 	const UE::Interchange::FAttributeKey Macro_CustomBuildNaniteKey = UE::Interchange::FAttributeKey(TEXT("BuildNanite"));
+	const UE::Interchange::FAttributeKey Macro_CustomAutoComputeLODScreenSizesKey = UE::Interchange::FAttributeKey(TEXT("AutoComputeLODScreenSizes"));
+
 
+	UE::Interchange::TArrayAttributeHelper<float> LODScreenSizes;
 	UE::Interchange::TArrayAttributeHelper<FString> SocketUids;
 
 protected:
diff --git a/Engine/Plugins/Interchange/Runtime/Source/Import/Private/Mesh/InterchangeStaticMeshFactory.cpp b/Engine/Plugins/Interchange/Runtime/Source/Import/Private/Mesh/InterchangeStaticMeshFactory.cpp
--- a/Engine/Plugins/Interchange/Runtime/Source/Import/Private/Mesh/InterchangeStaticMeshFactory.cpp
+++ b/Engine/Plugins/Interchange/Runtime/Source/Import/Private/Mesh/InterchangeStaticMeshFactory.cpp
@@ -450,8 +450,15 @@ UInterchangeFactoryBase::FImportAssetResult UInterchangeStaticMeshFactory::Impor
 
 #if WITH_EDITOR
 	{
+		// Default to AutoComputeLODScreenSizes in case the attribute is not set.
+		bool bAutoComputeLODScreenSize = true;
+		StaticMeshFactoryNode->GetCustomAutoComputeLODScreenSizes(bAutoComputeLODScreenSize);
+
+		TArray<float> LodScreenSizes;
+		StaticMeshFactoryNode->GetLODScreenSizes(LodScreenSizes);
+
 		const bool bIsAReimport = Arguments.ReimportObject != nullptr;
-		SetupSourceModelsSettings(*StaticMesh, LodMeshDescriptions, PrevLodCount, FinalLodCount, bIsAReimport);
+		SetupSourceModelsSettings(*StaticMesh, LodMeshDescriptions, bAutoComputeLODScreenSize, LodScreenSizes, PrevLodCount, FinalLodCount, bIsAReimport);
 
 		// SetupSourceModelsSettings can change the destination lightmap UV index
 		// Make sure the destination lightmap UV index on the factory node takes
@@ -804,12 +811,36 @@ void UInterchangeStaticMeshFactory::BuildFromMeshDescriptions(UStaticMesh& Stati
 }
 
 #if WITH_EDITORONLY_DATA
-void UInterchangeStaticMeshFactory::SetupSourceModelsSettings(UStaticMesh& StaticMesh, const TArray<FMeshDescription>& LodMeshDescriptions, int32 PreviousLodCount, int32 FinalLodCount, bool bIsAReimport)
+void UInterchangeStaticMeshFactory::SetupSourceModelsSettings(UStaticMesh& StaticMesh, const TArray<FMeshDescription>& LodMeshDescriptions, bool bAutoComputeLODScreenSizes, const TArray<float>& LodScreenSizes, int32 PreviousLodCount, int32 FinalLodCount, bool bIsAReimport)
 {
+	// Default LOD Screen Size
+	constexpr int32 LODIndex = 0;
+	float PreviousLODScreenSize = UStaticMesh::ComputeLODScreenSize(LODIndex);
+
+	// If no values are provided, then force AutoCompute
+	if (LodScreenSizes.IsEmpty())
+	{
+		bAutoComputeLODScreenSizes = true;
+	}
+	StaticMesh.bAutoComputeLODScreenSize = bAutoComputeLODScreenSizes;
+	
 	for (int32 LodIndex = 0; LodIndex < FinalLodCount; ++LodIndex)
 	{
 		FStaticMeshSourceModel& SrcModel = StaticMesh.GetSourceModel(LodIndex);
 
+		if (!bAutoComputeLODScreenSizes)
+		{
+			if (LodScreenSizes.IsValidIndex(LodIndex))
+			{
+				SrcModel.ScreenSize = LodScreenSizes[LodIndex]; 
+			}
+			else
+			{
+				SrcModel.ScreenSize = UStaticMesh::ComputeLODScreenSize(LodIndex, PreviousLODScreenSize);
+			}
+			PreviousLODScreenSize = SrcModel.ScreenSize.Default;
+		}
+
 		// Make sure that mesh descriptions for added LODs are kept as is when the mesh is built
 		if (LodIndex >= PreviousLodCount)
 		{
diff --git a/Engine/Plugins/Interchange/Runtime/Source/Import/Public/Mesh/InterchangeStaticMeshFactory.h b/Engine/Plugins/Interchange/Runtime/Source/Import/Public/Mesh/InterchangeStaticMeshFactory.h
--- a/Engine/Plugins/Interchange/Runtime/Source/Import/Public/Mesh/InterchangeStaticMeshFactory.h
+++ b/Engine/Plugins/Interchange/Runtime/Source/Import/Public/Mesh/InterchangeStaticMeshFactory.h
@@ -73,7 +73,7 @@ private:
 	void BuildFromMeshDescriptions(UStaticMesh& StaticMesh);
 
 #if WITH_EDITORONLY_DATA
-	void SetupSourceModelsSettings(UStaticMesh& StaticMesh, const TArray<FMeshDescription>& LodMeshDescriptions, int32 PreviousLodCount, int32 FinalLodCount, bool bIsAReimport);
+	void SetupSourceModelsSettings(UStaticMesh& StaticMesh, const TArray<FMeshDescription>& LodMeshDescriptions, bool bAutoComputeLODScreenSizes, const TArray<float>& LodScreenSizes, int32 PreviousLodCount, int32 FinalLodCount, bool bIsAReimport);
 #endif
 	struct FImportAssetObjectData
 	{
diff --git a/Engine/Plugins/Interchange/Runtime/Source/Pipelines/Private/InterchangeGenericStaticMeshPipeline.cpp b/Engine/Plugins/Interchange/Runtime/Source/Pipelines/Private/InterchangeGenericStaticMeshPipeline.cpp
--- a/Engine/Plugins/Interchange/Runtime/Source/Pipelines/Private/InterchangeGenericStaticMeshPipeline.cpp
+++ b/Engine/Plugins/Interchange/Runtime/Source/Pipelines/Private/InterchangeGenericStaticMeshPipeline.cpp
@@ -452,6 +452,7 @@ UInterchangeStaticMeshFactoryNode* UInterchangeGenericMeshPipeline::CreateStatic
 	StaticMeshFactoryNode->SetCustomUseFullPrecisionUVs(CommonMeshesProperties->bUseFullPrecisionUVs);
 	StaticMeshFactoryNode->SetCustomUseBackwardsCompatibleF16TruncUVs(CommonMeshesProperties->bUseBackwardsCompatibleF16TruncUVs);
 	StaticMeshFactoryNode->SetCustomRemoveDegenerates(CommonMeshesProperties->bRemoveDegenerates);
+	
 	//Static meshes build options
 	StaticMeshFactoryNode->SetCustomBuildReversedIndexBuffer(bBuildReversedIndexBuffer);
 	StaticMeshFactoryNode->SetCustomGenerateLightmapUVs(bGenerateLightmapUVs);
@@ -465,6 +466,9 @@ UInterchangeStaticMeshFactoryNode* UInterchangeGenericMeshPipeline::CreateStatic
 	StaticMeshFactoryNode->SetCustomDistanceFieldReplacementMesh(DistanceFieldReplacementMesh.Get());
 	StaticMeshFactoryNode->SetCustomMaxLumenMeshCards(MaxLumenMeshCards);
 	StaticMeshFactoryNode->SetCustomBuildNanite(bBuildNanite);
+	StaticMeshFactoryNode->SetCustomAutoComputeLODScreenSizes(bAutoComputeLODScreenSizes);
+	StaticMeshFactoryNode->SetLODScreenSizes(LODScreenSizes);
+	
 
 	if (ReferencingMeshInstanceUids.Num() > 0)
 	{
diff --git a/Engine/Plugins/Interchange/Runtime/Source/Pipelines/Public/InterchangeGenericMeshPipeline.h b/Engine/Plugins/Interchange/Runtime/Source/Pipelines/Public/InterchangeGenericMeshPipeline.h
--- a/Engine/Plugins/Interchange/Runtime/Source/Pipelines/Public/InterchangeGenericMeshPipeline.h
+++ b/Engine/Plugins/Interchange/Runtime/Source/Pipelines/Public/InterchangeGenericMeshPipeline.h
@@ -53,9 +53,17 @@ public:
 	bool bCombineStaticMeshes = false;
 
 	/** The LOD group that will be assigned to this mesh. */
-	UPROPERTY(EditAnywhere, BlueprintReadWrite, Category = "Static Meshes", meta=(DisplayName="LOD Group"))
+	UPROPERTY(EditAnywhere, BlueprintReadWrite, Category = "Static Meshes", meta=(DisplayName = "LOD Group"))
 	FName LodGroup = NAME_None;
 
+	/** If enabled, LOD Screen Sizes would be auto-computed. */
+	UPROPERTY(EditAnywhere, BlueprintReadWrite, Category = "Static Meshes")
+	bool bAutoComputeLODScreenSizes = true;
+
+	/** This setting is only used if the Auto Compute LOD Screen Sizes setting is disabled. */
+	UPROPERTY(EditAnywhere, BlueprintReadWrite, Category = "Static Meshes", meta=(EditCondition = "!bAutoComputeLODScreenSizes", DisplayName = "LOD Screen Sizes"))
+	TArray<float> LODScreenSizes;
+
 	UPROPERTY(meta = (DeprecatedProperty, DeprecationMessage = "Use Collision instead."))
 	bool bImportCollision_DEPRECATED = true;
 
diff --git a/Engine/Source/Runtime/Engine/Classes/Engine/StaticMesh.h b/Engine/Source/Runtime/Engine/Classes/Engine/StaticMesh.h
--- a/Engine/Source/Runtime/Engine/Classes/Engine/StaticMesh.h
+++ b/Engine/Source/Runtime/Engine/Classes/Engine/StaticMesh.h
@@ -789,7 +789,10 @@ public:
 	ENGINE_API int32 GetMinLODIdx(bool bForceLowestLODIdx = false) const;
 	ENGINE_API int32 GetDefaultMinLOD() const;
 	ENGINE_API void SetMinLODIdx(int32 InMinLOD);
-
+	
+	/** Computes the LOD Screen Size based on the LODIndex while keeping it within tolerance from the previous LODScreenSize*/
+	ENGINE_API static float ComputeLODScreenSize(int32 LODIndex, float PreviousLODScreenSize = -1.0f);
+	
 	ENGINE_API static void OnLodStrippingQualityLevelChanged(IConsoleVariable* Variable);
 
 	/** Minimum LOD to use for rendering.  This is the default setting for the mesh and can be overridden by component settings. */
diff --git a/Engine/Source/Runtime/Engine/Private/StaticMesh.cpp b/Engine/Source/Runtime/Engine/Private/StaticMesh.cpp
--- a/Engine/Source/Runtime/Engine/Private/StaticMesh.cpp
+++ b/Engine/Source/Runtime/Engine/Private/StaticMesh.cpp
@@ -192,6 +192,19 @@ static void FillMaterialName(const TArray<FStaticMaterial>& StaticMaterials, TMa
 }
 #endif
 
+namespace UE::StaticMesh::Private::Constants
+{
+	// Default Highest LOD Screen Size
+	constexpr float LOD0ScreenSize = 2.0f;
+
+	// Arbitrary constant used as a base in Pow(K, LODIndex) that achieves much the same progression as a
+	// conversion of the old 1 / (MaxLODs * LODIndex) passed through the newer bounds computation.
+	// i.e. this achieves much the same results, but is still fairly arbitrary.
+	constexpr float AutoComputeLODPowerBase = 0.75f;
+
+	constexpr float Tolerance = 0.01f;
+}
+
 /*-----------------------------------------------------------------------------
 	FStaticMeshAsyncBuildWorker
 -----------------------------------------------------------------------------*/
@@ -2542,20 +2555,17 @@ void FStaticMeshRenderData::ResolveSectionInfo(UStaticMesh* Owner)
 			Section.bForceOpaque = Info.bForceOpaque;
 		}
 
-		// Arbitrary constant used as a base in Pow(K, LODIndex) that achieves much the same progression as a
-		// conversion of the old 1 / (MaxLODs * LODIndex) passed through the newer bounds computation.
-		// i.e. this achieves much the same results, but is still fairly arbitrary.
-		const float AutoComputeLODPowerBase = 0.75f;
-
 		if (Owner->bAutoComputeLODScreenSize)
 		{
+			using namespace UE::StaticMesh::Private;
+
 			if (LODIndex == 0)
 			{
-				ScreenSize[LODIndex].Default = 2.0f;
+				ScreenSize[LODIndex].Default = Constants::LOD0ScreenSize;
 			}
 			else if(LOD.MaxDeviation <= 0.0f)
 			{
-				ScreenSize[LODIndex].Default = FMath::Pow(AutoComputeLODPowerBase, LODIndex);
+				ScreenSize[LODIndex].Default = FMath::Pow(Constants::AutoComputeLODPowerBase, LODIndex);
 			}
 			else
 			{
@@ -2590,15 +2600,10 @@ void FStaticMeshRenderData::ResolveSectionInfo(UStaticMesh* Owner)
 		}
 		else
 		{
-			check(LODIndex > 0);
-
 			// No valid source model and we're not auto-generating. Auto-generate in this case
 			// because we have nothing else to go on.
-			const float Tolerance = 0.01f;
-			float AutoDisplayFactor = FMath::Pow(AutoComputeLODPowerBase, LODIndex);
-
-			// Make sure this fits in with the previous LOD
-			ScreenSize[LODIndex].Default = FMath::Clamp(AutoDisplayFactor, 0.0f, ScreenSize[LODIndex-1].Default - Tolerance);
+			check(LODIndex > 0);
+			ScreenSize[LODIndex].Default = UStaticMesh::ComputeLODScreenSize(LODIndex, ScreenSize[LODIndex - 1].Default);
 		}
 	}
 	for (; LODIndex < MAX_STATIC_MESH_LODS; ++LODIndex)
@@ -2901,6 +2906,26 @@ void FStaticMeshLODSettings::ReadEntry(FStaticMeshLODGroup& Group, FString Entry
 	}
 }
 
+float UStaticMesh::ComputeLODScreenSize(int32 LODIndex, float PreviousLODScreenSize /*= -1.0f*/)
+{
+	using namespace UE::StaticMesh::Private;
+
+	if (PreviousLODScreenSize < 0.0f)
+	{
+		// Provides a good upper bound
+		PreviousLODScreenSize = Constants::LOD0ScreenSize;
+	}
+
+	if (LODIndex <= 0)
+	{
+		return Constants::LOD0ScreenSize;
+	}
+
+	const float AutoDisplayFactor = FMath::Pow(Constants::AutoComputeLODPowerBase, LODIndex);
+	// Make sure this fits in with the previous LOD
+	return FMath::Clamp(AutoDisplayFactor, 0.0f, PreviousLODScreenSize - Constants::Tolerance);
+}
+
 void FStaticMeshLODSettings::GetLODGroupNames(TArray<FName>& OutNames) const
 {
 	for (TMap<FName,FStaticMeshLODGroup>::TConstIterator It(Groups); It; ++It)
