Interchange: dialog import preview missing the import context.

#rnx

diff --git a/Engine/Plugins/Interchange/Editor/Source/Pipelines/Private/SInterchangePipelineConfigurationDialog.cpp b/Engine/Plugins/Interchange/Editor/Source/Pipelines/Private/SInterchangePipelineConfigurationDialog.cpp
--- a/Engine/Plugins/Interchange/Editor/Source/Pipelines/Private/SInterchangePipelineConfigurationDialog.cpp
+++ b/Engine/Plugins/Interchange/Editor/Source/Pipelines/Private/SInterchangePipelineConfigurationDialog.cpp
@@ -1571,6 +1571,7 @@ void SInterchangePipelineConfigurationDialog::UpdatePreviewContainer(bool bUpdat
 		//Duplicate the pipeline because ScriptedExecutePipeline is not const
 		if (UInterchangePipelineBase* DuplicatedPipeline = DuplicateObject<UInterchangePipelineBase>(PipelineItem->Pipeline, GetTransientPackage()))
 		{
+			DuplicatedPipeline->TransferAdjustSettings(PipelineItem->Pipeline);
 			DuplicatedPipeline->SetResultsContainer(Results);
 			DuplicatedPipeline->ScriptedExecutePipeline(PreviewNodeContainer, SourceDatas, FString());
 			ClearObjectFlags(DuplicatedPipeline);
diff --git a/Engine/Source/Runtime/Interchange/Core/Private/InterchangePipelineBase.cpp b/Engine/Source/Runtime/Interchange/Core/Private/InterchangePipelineBase.cpp
--- a/Engine/Source/Runtime/Interchange/Core/Private/InterchangePipelineBase.cpp
+++ b/Engine/Source/Runtime/Interchange/Core/Private/InterchangePipelineBase.cpp
@@ -63,7 +63,7 @@ void UInterchangePipelineBase::AdjustSettingsFromCache()
 	AdjustSettingsForContext(CacheContextParam);
 }
 
-void UInterchangePipelineBase::TransferAdjustSettings(UInterchangePipelineBase* SourcePipeline)
+void UInterchangePipelineBase::TransferAdjustSettings(const UInterchangePipelineBase* SourcePipeline)
 {
 	CacheContextParam = SourcePipeline->CacheContextParam;
 	CachePropertiesStates = SourcePipeline->CachePropertiesStates;
diff --git a/Engine/Source/Runtime/Interchange/Core/Public/InterchangePipelineBase.h b/Engine/Source/Runtime/Interchange/Core/Public/InterchangePipelineBase.h
--- a/Engine/Source/Runtime/Interchange/Core/Public/InterchangePipelineBase.h
+++ b/Engine/Source/Runtime/Interchange/Core/Public/InterchangePipelineBase.h
@@ -332,7 +332,7 @@ public:
 	INTERCHANGECORE_API virtual void AdjustSettingsFromCache();
 
 	/** Transfer the source pipeline adjust settings to this pipeline. */
-	INTERCHANGECORE_API void TransferAdjustSettings(UInterchangePipelineBase* SourcePipeline);
+	INTERCHANGECORE_API void TransferAdjustSettings(const UInterchangePipelineBase* SourcePipeline);
 
 	INTERCHANGECORE_API void SetShowEssentialsMode(bool bShowEssentialsModeValue)
 	{
