Interchange Import Dialog
- Warning Icon on Import Button
- Added ability to show conflict warnings on the cards view guarded by a CCVar

diff --git a/Engine/Plugins/Interchange/Editor/Source/Pipelines/Private/SInterchangeAssetCard.cpp b/Engine/Plugins/Interchange/Editor/Source/Pipelines/Private/SInterchangeAssetCard.cpp
--- a/Engine/Plugins/Interchange/Editor/Source/Pipelines/Private/SInterchangeAssetCard.cpp
+++ b/Engine/Plugins/Interchange/Editor/Source/Pipelines/Private/SInterchangeAssetCard.cpp
@@ -3,6 +3,7 @@
 
 #include "GameFramework/Actor.h"
 #include "InterchangeEditorPipelineStyle.h"
+#include "InterchangePipelineBase.h"
 #include "Nodes/InterchangeBaseNodeContainer.h"
 #include "Nodes/InterchangeFactoryBaseNode.h"
 #include "SSimpleButton.h"
@@ -14,6 +15,7 @@
 #include "Widgets/Layout/SBorder.h"
 #include "Widgets/Layout/SBox.h"
 #include "Widgets/Layout/SSeparator.h"
+#include "Widgets/Layout/SSpacer.h"
 #include "Widgets/Layout/SSplitter.h"
 #include "Widgets/SBoxPanel.h"
 #include "Widgets/Text/STextBlock.h"
@@ -67,6 +69,21 @@ void SInterchangeAssetCard::RefreshCard(UInterchangeBaseNodeContainer* InPreview
 		});
 }
 
+bool SInterchangeAssetCard::RefreshHasConflicts(const TArray<FInterchangeConflictInfo>& InConflictInfos)
+{
+	for (const FInterchangeConflictInfo& ConflictInfo : InConflictInfos)
+	{
+		if (ConflictInfo.AffectedAssetClasses.Contains(AssetClass))
+		{
+			bHasConflictWarnings = true;
+			return true;
+		}
+	}
+
+	bHasConflictWarnings = false;
+	return false;
+}
+
 void SInterchangeAssetCard::Construct(const FArguments& InArgs)
 {
 	//Make sure we have a valid node container
@@ -90,7 +107,7 @@ void SInterchangeAssetCard::Construct(const FArguments& InArgs)
 
 	FText CardImportText = FText::Format(LOCTEXT("CardImportText", "Import {0}"), FText::FromString(AssetClass->GetName()));
 
-	FName IconWarningName = "InterchangeAssetCardIcon.Warning";
+	FName IconWarningName = "Icons.Alert.Solid";
 	const FSlateIcon IconWarning = FSlateIconFinder::FindIcon(IconWarningName);
 
 	const ISlateStyle* InterchangeEditorPipelineStyle = FSlateStyleRegistry::FindSlateStyle("InterchangeEditorPipelineStyle");
@@ -121,13 +138,17 @@ void SInterchangeAssetCard::Construct(const FArguments& InArgs)
 			[
 				SNew(SHorizontalBox)
 				+ SHorizontalBox::Slot()
+				.VAlign(VAlign_Center)
+				.HAlign(HAlign_Center)
 				.AutoWidth()
 				.Padding(8.0f, 4.0f, 4.0f, 4.0f)
 				[
 					CardIconWidget.ToSharedRef()
 				]
 				+ SHorizontalBox::Slot()
-				.FillWidth(1.0f)
+				.VAlign(VAlign_Center)
+				.HAlign(HAlign_Left)
+				.AutoWidth()
 				.Padding(4.0f, 4.0f, 8.0f, 4.0f)
 				[
 					SNew(STextBlock)
@@ -141,6 +162,39 @@ void SInterchangeAssetCard::Construct(const FArguments& InArgs)
 							return FText::FromString(CardTooltip);
 						})
 				]
+				+ SHorizontalBox::Slot()
+				.FillWidth(1.0f)
+				[
+					SNew(SSpacer)
+				]
+				+ SHorizontalBox::Slot()
+				.AutoWidth()
+				.Padding(8.0f, 4.0f, 8.0f, 4.0f)
+				[
+					SNew(SHorizontalBox)
+					.Visibility_Lambda([this]() {return bHasConflictWarnings ? EVisibility::Visible : EVisibility::Collapsed; })
+					.ToolTipText(LOCTEXT("ConflictWarningTooltipText", "There are some conflicts generated while importing the source file. Go to Conflicts Section in Advanced Settings for more details."))	
+					+ SHorizontalBox::Slot()
+					.VAlign(VAlign_Center)
+					.HAlign(HAlign_Center)
+					.AutoWidth()
+					[
+						SNew(SImage)
+						.Image(IconWarning.GetOptionalIcon())
+						.ColorAndOpacity(FStyleColors::Warning)
+					]
+					+ SHorizontalBox::Slot()
+					.VAlign(VAlign_Center)
+					.HAlign(HAlign_Left)
+					.AutoWidth()
+					.Padding(4.0f, 4.0f, 8.0f, 4.0f)
+					[
+						SNew(STextBlock)
+						.Text(LOCTEXT("ConflictWarningText", "Conflict Warnings"))
+						.ColorAndOpacity(FStyleColors::Warning)
+					]
+					
+				]
 			]
 		]
 		+ SVerticalBox::Slot()
@@ -181,17 +235,17 @@ void SInterchangeAssetCard::Construct(const FArguments& InArgs)
 					[
 						SNew(SCheckBox)
 						.IsChecked_Lambda([this]()
+						{
+							if (ShouldImportAssetType.Execute())
 							{
-								if (ShouldImportAssetType.Execute())
-								{
-									return ECheckBoxState::Checked;
-								}
-								return ECheckBoxState::Unchecked;
-							})
+								return ECheckBoxState::Checked;
+							}
+							return ECheckBoxState::Unchecked;
+						})
 						.OnCheckStateChanged_Lambda([this](ECheckBoxState CheckBoxState)
-							{
-								OnImportAssetTypeChanged.Execute(CheckBoxState == ECheckBoxState::Checked);
-							})
+						{
+							OnImportAssetTypeChanged.Execute(CheckBoxState == ECheckBoxState::Checked);
+						})
 					]
 				]
 			]
diff --git a/Engine/Plugins/Interchange/Editor/Source/Pipelines/Private/SInterchangeAssetCard.h b/Engine/Plugins/Interchange/Editor/Source/Pipelines/Private/SInterchangeAssetCard.h
--- a/Engine/Plugins/Interchange/Editor/Source/Pipelines/Private/SInterchangeAssetCard.h
+++ b/Engine/Plugins/Interchange/Editor/Source/Pipelines/Private/SInterchangeAssetCard.h
@@ -11,6 +11,7 @@
 #include "Widgets/Views/STableViewBase.h"
 
 class UInterchangeBaseNodeContainer;
+struct FInterchangeConflictInfo;
 
 DECLARE_DELEGATE_OneParam(FInterchangeCard_OnImportAssetTypeChanged, bool)
 DECLARE_DELEGATE_RetVal(bool, FInterchangeCard_ShouldImportAssetType)
@@ -38,12 +39,15 @@ public:
 
 	void RefreshCard(UInterchangeBaseNodeContainer* InPreviewNodeContainer);
 
+	bool RefreshHasConflicts(const TArray<FInterchangeConflictInfo>& InConflictInfos);
+
 protected:
 
 	int32 CardAssetCount = 0;
 	int32 CardAssetToImportCount = 0;
 	int32 CardAssetDisabledCount = 0;
 	FString CardTooltip;
+	bool bHasConflictWarnings = false;
 
 	/** The factory asset class so we can know which kind of Unreal asset this card is for. */
 	UClass* AssetClass = nullptr;
@@ -53,7 +57,6 @@ protected:
 
 	/** Delegate to invoke when should import asset type value change. */
 	FInterchangeCard_OnImportAssetTypeChanged OnImportAssetTypeChanged;
-
 };
 
 class SInterchangeAssetCardList : public SCompoundWidget
diff --git a/Engine/Plugins/Interchange/Editor/Source/Pipelines/Private/SInterchangePipelineConfigurationDialog.cpp b/Engine/Plugins/Interchange/Editor/Source/Pipelines/Private/SInterchangePipelineConfigurationDialog.cpp
--- a/Engine/Plugins/Interchange/Editor/Source/Pipelines/Private/SInterchangePipelineConfigurationDialog.cpp
+++ b/Engine/Plugins/Interchange/Editor/Source/Pipelines/Private/SInterchangePipelineConfigurationDialog.cpp
@@ -60,6 +60,13 @@ static FAutoConsoleVariableRef CCvarInterchangeDefaultHideCardsView(
 	TEXT("Whether the import dialog should hide the basic cards view."),
 	ECVF_Default);
 
+static bool GInterchangeShowConflictWarningsOnCardsView = true;
+static FAutoConsoleVariableRef CCvarInterchangeShowConflictWarningOnCardsView(
+	TEXT("Interchange.FeatureFlags.Import.ShowConflictWarningsOnCardsView"),
+	GInterchangeShowConflictWarningsOnCardsView,
+	TEXT("Whether the import conflict warnings will be shown on cards view."),
+	ECVF_Default);
+
 constexpr double SplitterSize = 12.0;
 constexpr double AdvancedUIRatio = 1.2;
 
@@ -1016,6 +1023,7 @@ void SInterchangePipelineConfigurationDialog::Construct(const FArguments& InArgs
 							return bShowCards ? FMargin(8.0f, 0.0f, 0.0f, 0.0f) : FMargin(0.0f, 0.0f, 0.0f, 0.0f);
 						})
 					.Padding(8.0f, 0.0f, 0.0f, 0.0f)
+					.Clipping(EWidgetClipping::ClipToBounds)
 					[
 						MainBodyAdvanced
 					]
@@ -1114,7 +1122,7 @@ void SInterchangePipelineConfigurationDialog::Construct(const FArguments& InArgs
 		FString KeyName = CurrentStackName.ToString() + TEXT("_LastSelectedPipeline");
 		if (GConfig->GetString(TEXT("InterchangeSelectPipeline"), *KeyName, LastPipelineName, GEditorPerProjectIni))
 		{
-			for (TSharedPtr<FInterchangePipelineItemType> PipelineItem : PipelineListViewItems)
+			for (TSharedPtr<FInterchangePipelineItemType> PipelineItem: PipelineListViewItems)
 			{
 				FString PipelineItemName = PipelineItem->Pipeline->GetClass()->GetName();
 				if (PipelineItemName.Equals(LastPipelineName))
@@ -1129,6 +1137,11 @@ void SInterchangePipelineConfigurationDialog::Construct(const FArguments& InArgs
 		{
 			PipelinesListView->SetSelection(PipelineListViewItems[0], ESelectInfo::Direct);
 		}
+
+		if (GInterchangeShowConflictWarningsOnCardsView)
+		{
+			RefreshCardsViewList();
+		}
 	}
 }
 
@@ -1145,14 +1158,17 @@ bool SInterchangePipelineConfigurationDialog::IsPropertyVisible(const FPropertyA
 const FSlateBrush* SInterchangePipelineConfigurationDialog::GetImportButtonIcon() const
 {
 	const FSlateBrush* TypeIcon = nullptr;
-	for (TSharedPtr<FInterchangePipelineItemType> PipelineItem : PipelineListViewItems)
+	if (bShowSettings || GInterchangeShowConflictWarningsOnCardsView)
 	{
-		if (PipelineItem.IsValid() && PipelineItem->Pipeline)
+		for (TSharedPtr<FInterchangePipelineItemType> PipelineItem : PipelineListViewItems)
 		{
-			if (PipelineItem->ConflictInfos.Num() > 0)
+			if (PipelineItem.IsValid() && PipelineItem->Pipeline)
 			{
-				const FSlateIcon SlateIcon = FSlateIconFinder::FindIcon("Icons.Warning");
-				return SlateIcon.GetOptionalIcon();
+				if (PipelineItem->ConflictInfos.Num() > 0)
+				{
+					const FSlateIcon SlateIcon = FSlateIconFinder::FindIcon("Icons.Warning");
+					return SlateIcon.GetOptionalIcon();
+				}
 			}
 		}
 	}
@@ -1405,16 +1421,26 @@ void SInterchangePipelineConfigurationDialog::RefreshStack(bool bStackSelectionC
 			}
 		}
 	}
-
-	if (bShowSettings)
+	//Select the first pipeline
+	if (PipelineListViewItems.Num() > 0)
 	{
-		//Select the first pipeline
-		if (PipelineListViewItems.Num() > 0)
+		CurrentPipelineIndex = PipelineListViewItems.IsValidIndex(CurrentPipelineIndex) ? CurrentPipelineIndex : 0;
+		if (bShowSettings)
 		{
-			CurrentPipelineIndex = PipelineListViewItems.IsValidIndex(CurrentPipelineIndex) ? CurrentPipelineIndex : 0;
 			PipelinesListView->SetSelection(PipelineListViewItems[CurrentPipelineIndex], ESelectInfo::Direct);
+			PipelinesListView->RequestListRefresh();
+		}
+		else
+		{
+			for (TSharedPtr<FInterchangePipelineItemType>& PiplineListViewItem : PipelineListViewItems)
+			{
+				PiplineListViewItem->ConflictInfos.Reset();
+				if (PiplineListViewItem->ReimportObject)
+				{
+					PiplineListViewItem->ConflictInfos = PiplineListViewItem->Pipeline->GetConflictInfos(PiplineListViewItem->ReimportObject, PiplineListViewItem->Container, PiplineListViewItem->SourceData);
+				}
+			}
 		}
-		PipelinesListView->RequestListRefresh();
 	}
 
 	//Update the cards
@@ -1570,6 +1596,17 @@ void SInterchangePipelineConfigurationDialog::FillAssetCardsList()
 							EnableDataPerFactoryNodeClass.FindChecked(FactoryNodeClass).bEnable = bNewEnabledValue;
 						});
 
+				if (GInterchangeShowConflictWarningsOnCardsView
+					&& PipelineListViewItems.Num() > 0)
+				{
+					for (TSharedPtr<FInterchangePipelineItemType>& PipelineListViewItem : PipelineListViewItems)
+					{
+						if (AssetCard->RefreshHasConflicts(PipelineListViewItem->ConflictInfos))
+						{
+							break;
+						}
+					}
+				}
 				AssetCards.Add(AssetCard);
 			}
 		}
diff --git a/Engine/Plugins/Interchange/Runtime/Source/Pipelines/Private/InterchangeGenericAssetsPipeline.cpp b/Engine/Plugins/Interchange/Runtime/Source/Pipelines/Private/InterchangeGenericAssetsPipeline.cpp
--- a/Engine/Plugins/Interchange/Runtime/Source/Pipelines/Private/InterchangeGenericAssetsPipeline.cpp
+++ b/Engine/Plugins/Interchange/Runtime/Source/Pipelines/Private/InterchangeGenericAssetsPipeline.cpp
@@ -566,6 +566,7 @@ void UInterchangeGenericAssetsPipeline::CreateMaterialConflict(UStaticMesh* Stat
 		MaterialConflict.Description = TEXT("There is some unmatched materials");
 		MaterialConflict.Pipeline = this;
 		MaterialConflict.UniqueId = FGuid::NewGuid();
+		MaterialConflict.AffectedAssetClasses.Add(UMaterialInterface::StaticClass());
 
 		//Cache the data so we do not have to redo the works when we will show the conflict
 		MaterialConflictData.ConflictUniqueId = MaterialConflict.UniqueId;
@@ -975,6 +976,7 @@ void UInterchangeGenericAssetsPipeline::CreateSkeletonConflict(USkeleton* Specif
 	SkeletonConflict.Description = TEXT("Imported skeleton is incompatible with the asset skeleton");
 	SkeletonConflict.Pipeline = this;
 	SkeletonConflict.UniqueId = FGuid::NewGuid();
+	SkeletonConflict.AffectedAssetClasses.Add(USkeletalMesh::StaticClass());
 
 	SkeletonConflictData.ConflictUniqueId = SkeletonConflict.UniqueId;
 	SkeletonConflictData.ReimportObject = SkeletalMesh;
diff --git a/Engine/Source/Runtime/Interchange/Core/Public/InterchangePipelineBase.h b/Engine/Source/Runtime/Interchange/Core/Public/InterchangePipelineBase.h
--- a/Engine/Source/Runtime/Interchange/Core/Public/InterchangePipelineBase.h
+++ b/Engine/Source/Runtime/Interchange/Core/Public/InterchangePipelineBase.h
@@ -164,6 +164,7 @@ struct FInterchangeConflictInfo
 	FString Description;
 	FGuid UniqueId;
 	TObjectPtr<UInterchangePipelineBase> Pipeline = nullptr;
+	TSet<UClass*> AffectedAssetClasses;
 };
 
 class SInterchangeBaseConflictWidget : public SCompoundWidget
