USD: Add details customizations to the Interchange USD translator settings dialog.

Refactor "render context registry" into a couple of static functions on the USDCore plugin;
Cleanup usage of render context FNames.

diff --git a/Engine/Plugins/Importers/USDImporter/Source/USDSchemas/Private/Custom/MDLUSDShadeMaterialTranslator.cpp b/Engine/Plugins/Importers/USDImporter/Source/USDSchemas/Private/Custom/MDLUSDShadeMaterialTranslator.cpp
--- a/Engine/Plugins/Importers/USDImporter/Source/USDSchemas/Private/Custom/MDLUSDShadeMaterialTranslator.cpp
+++ b/Engine/Plugins/Importers/USDImporter/Source/USDSchemas/Private/Custom/MDLUSDShadeMaterialTranslator.cpp
@@ -35,7 +35,9 @@
 
 #define LOCTEXT_NAMESPACE "MDLUSDShadeMaterialTranslator"
 
+PRAGMA_DISABLE_DEPRECATION_WARNINGS
 FName FMdlUsdShadeMaterialTranslator::MdlRenderContext = TEXT("mdl");
+PRAGMA_ENABLE_DEPRECATION_WARNINGS
 
 namespace UE::MDLShadeTranslator::Private
 {
@@ -73,7 +75,7 @@ void FMdlUsdShadeMaterialTranslator::CreateAssets()
 	//   info:mdl:sourceAsset:subIdentifier -> Name of the material in the MDL file
 	//   inputs -> material parameters
 
-	if (Context->RenderContext != MdlRenderContext)
+	if (Context->RenderContext != UnrealIdentifiers::MdlRenderContext)
 	{
 		Super::CreateAssets();
 		return;
@@ -102,7 +104,7 @@ void FMdlUsdShadeMaterialTranslator::CreateAssets()
 		return;
 	}
 
-	const pxr::TfToken MdlToken = UnrealToUsd::ConvertToken(*MdlRenderContext.ToString()).Get();
+	const static pxr::TfToken MdlToken = UnrealToUsd::ConvertToken(*UnrealIdentifiers::MdlRenderContext.ToString()).Get();
 
 	pxr::UsdShadeShader SurfaceShader = ShadeMaterial.ComputeSurfaceSource(MdlToken);
 
diff --git a/Engine/Plugins/Importers/USDImporter/Source/USDSchemas/Private/Custom/MDLUSDShadeMaterialTranslator.h b/Engine/Plugins/Importers/USDImporter/Source/USDSchemas/Private/Custom/MDLUSDShadeMaterialTranslator.h
--- a/Engine/Plugins/Importers/USDImporter/Source/USDSchemas/Private/Custom/MDLUSDShadeMaterialTranslator.h
+++ b/Engine/Plugins/Importers/USDImporter/Source/USDSchemas/Private/Custom/MDLUSDShadeMaterialTranslator.h
@@ -17,6 +17,7 @@ class FMdlUsdShadeMaterialTranslator : public FMaterialXUsdShadeMaterialTranslat
 	using Super = FMaterialXUsdShadeMaterialTranslator;
 
 public:
+	UE_DEPRECATED(5.5, "Use UnrealIdentifiers::MdlRenderContext instead");
 	static FName MdlRenderContext;
 
 public:
diff --git a/Engine/Plugins/Importers/USDImporter/Source/USDSchemas/Private/USDSchemaTranslator.cpp b/Engine/Plugins/Importers/USDImporter/Source/USDSchemas/Private/USDSchemaTranslator.cpp
--- a/Engine/Plugins/Importers/USDImporter/Source/USDSchemas/Private/USDSchemaTranslator.cpp
+++ b/Engine/Plugins/Importers/USDImporter/Source/USDSchemas/Private/USDSchemaTranslator.cpp
@@ -2,7 +2,9 @@
 
 #include "USDSchemaTranslator.h"
 
+#include "UnrealUSDWrapper.h"
 #include "USDInfoCache.h"
+#include "USDMaterialUtils.h"
 #include "USDMemory.h"
 #include "USDSchemasModule.h"
 #include "USDTypesConversion.h"
@@ -63,18 +65,38 @@ TSharedPtr<FUsdSchemaTranslator> FUsdSchemaTranslatorRegistry::CreateTranslatorF
 	return {};
 }
 
+PRAGMA_DISABLE_DEPRECATION_WARNINGS
 FUsdRenderContextRegistry::FUsdRenderContextRegistry()
 {
-#if USE_USD_SDK
-	LLM_SCOPE_BYTAG(Usd);
+}
+
+void FUsdRenderContextRegistry::Register(const FName& RenderContextToken)
+{
+	UsdUnreal::MaterialUtils::RegisterRenderContext(RenderContextToken);
+}
 
-	UniversalRenderContext = FName(UsdToUnreal::ConvertToken(pxr::UsdShadeTokens->universalRenderContext));
-	Register(UniversalRenderContext);
+void FUsdRenderContextRegistry::Unregister(const FName& RenderContextToken)
+{
+	UsdUnreal::MaterialUtils::UnregisterRenderContext(RenderContextToken);
+}
 
-	UnrealRenderContext = FName(UsdToUnreal::ConvertToken(UnrealIdentifiers::Unreal));
-	Register(UnrealRenderContext);
-#endif	  // #if USE_USD_SDK
+const TSet<FName>& FUsdRenderContextRegistry::GetRenderContexts() const
+{
+	static TSet<FName> TempValue;
+	TempValue = TSet<FName>{UsdUnreal::MaterialUtils::GetRegisteredRenderContexts()};
+	return TempValue;
+}
+
+const FName& FUsdRenderContextRegistry::GetUniversalRenderContext() const
+{
+	return UnrealIdentifiers::UniversalRenderContext;
+}
+
+const FName& FUsdRenderContextRegistry::GetUnrealRenderContext() const
+{
+	return UnrealIdentifiers::UnrealRenderContext;
 }
+PRAGMA_ENABLE_DEPRECATION_WARNINGS
 
 FRegisteredSchemaTranslatorHandle FUsdSchemaTranslatorRegistry::Register(const FString& SchemaName, FCreateTranslator CreateFunction)
 {
@@ -151,16 +173,14 @@ FUsdSchemaTranslationContext::FUsdSchemaTranslationContext(const UE::FUsdStage&
 	: Stage(InStage)
 	, UsdAssetCache(NewObject<UUsdAssetCache3>())
 {
-	IUsdSchemasModule& UsdSchemasModule = FModuleManager::Get().LoadModuleChecked<IUsdSchemasModule>(TEXT("USDSchemas"));
-	RenderContext = UsdSchemasModule.GetRenderContextRegistry().GetUniversalRenderContext();
+	RenderContext = UnrealIdentifiers::UniversalRenderContext;
 }
 
 FUsdSchemaTranslationContext::FUsdSchemaTranslationContext(const UE::FUsdStage& InStage, UUsdAssetCache3& InAssetCache)
 	: Stage(InStage)
 	, UsdAssetCache(&InAssetCache)
 {
-	IUsdSchemasModule& UsdSchemasModule = FModuleManager::Get().LoadModuleChecked<IUsdSchemasModule>(TEXT("USDSchemas"));
-	RenderContext = UsdSchemasModule.GetRenderContextRegistry().GetUniversalRenderContext();
+	RenderContext = UnrealIdentifiers::UniversalRenderContext;
 }
 
 FUsdSchemaTranslatorRegistry::FSchemaTranslatorsStack* FUsdSchemaTranslatorRegistry::FindSchemaTranslatorStack(const FString& SchemaName)
diff --git a/Engine/Plugins/Importers/USDImporter/Source/USDSchemas/Private/USDSchemasModule.cpp b/Engine/Plugins/Importers/USDImporter/Source/USDSchemas/Private/USDSchemasModule.cpp
--- a/Engine/Plugins/Importers/USDImporter/Source/USDSchemas/Private/USDSchemasModule.cpp
+++ b/Engine/Plugins/Importers/USDImporter/Source/USDSchemas/Private/USDSchemasModule.cpp
@@ -45,7 +45,6 @@ public:
 		LLM_SCOPE_BYTAG(Usd);
 
 		FUsdSchemaTranslatorRegistry& Registry = GetTranslatorRegistry();
-		FUsdRenderContextRegistry& ShaderRegistry = GetRenderContextRegistry();
 
 		// Register the default translators
 		TranslatorHandles = {
@@ -65,7 +64,7 @@ public:
 			Registry.Register<FUsdVolVolumeTranslator>(TEXT("UsdVolVolume"))};
 
 #if WITH_EDITOR
-		ShaderRegistry.Register(UnrealIdentifiers::MaterialXRenderContext);
+		UsdUnreal::MaterialUtils::RegisterRenderContext(UnrealIdentifiers::MaterialXRenderContext);
 		TranslatorHandles.Add(Registry.Register<FMaterialXUsdShadeMaterialTranslator>(TEXT("UsdShadeMaterial")));
 
 		// Creating skeletal meshes technically works in Standalone mode, but by checking for this we artificially block it
@@ -73,18 +72,18 @@ public:
 		// static meshes, at least.
 		if (GIsEditor)
 		{
-			TranslatorHandles.Append({
-				Registry.Register<FUsdSkelSkeletonTranslator>(TEXT("UsdSkelSkeleton")),
-				Registry.Register<FUsdGroomTranslator>(TEXT("UsdGeomXformable")),
-				// The GeometryCacheTranslator also works on UsdGeomXformable through the GroomTranslator
-				Registry.Register<FUsdGeometryCacheTranslator>(TEXT("UsdGeomMesh")),
-				// It doesn't seem possible to create SoundWave assets at runtime at the moment, for whatever reason
-				Registry.Register<FUsdMediaSpatialAudioTranslator>(TEXT("UsdMediaSpatialAudio"))
-			});
+			TranslatorHandles.Append(
+				{Registry.Register<FUsdSkelSkeletonTranslator>(TEXT("UsdSkelSkeleton")),
+				 Registry.Register<FUsdGroomTranslator>(TEXT("UsdGeomXformable")),
+				 // The GeometryCacheTranslator also works on UsdGeomXformable through the GroomTranslator
+				 Registry.Register<FUsdGeometryCacheTranslator>(TEXT("UsdGeomMesh")),
+				 // It doesn't seem possible to create SoundWave assets at runtime at the moment, for whatever reason
+				 Registry.Register<FUsdMediaSpatialAudioTranslator>(TEXT("UsdMediaSpatialAudio"))}
+			);
 
 			if (IMDLImporterModule* MDLImporterModule = FModuleManager::Get().LoadModulePtr<IMDLImporterModule>(TEXT("MDLImporter")))
 			{
-				ShaderRegistry.Register(FMdlUsdShadeMaterialTranslator::MdlRenderContext);
+				UsdUnreal::MaterialUtils::RegisterRenderContext(UnrealIdentifiers::MdlRenderContext);
 				TranslatorHandles.Add(Registry.Register<FMdlUsdShadeMaterialTranslator>(TEXT("UsdShadeMaterial")));
 			}
 		}
@@ -98,7 +97,6 @@ public:
 	virtual void ShutdownModule() override
 	{
 		FUsdSchemaTranslatorRegistry& Registry = GetTranslatorRegistry();
-		FUsdRenderContextRegistry& ShaderRegistry = GetRenderContextRegistry();
 
 		for (const FRegisteredSchemaTranslatorHandle& TranslatorHandle : TranslatorHandles)
 		{
@@ -106,8 +104,8 @@ public:
 		}
 
 #if USE_USD_SDK && WITH_EDITOR
-		ShaderRegistry.Unregister(FMdlUsdShadeMaterialTranslator::MdlRenderContext);
-		ShaderRegistry.Unregister(UnrealIdentifiers::MaterialXRenderContext);
+		UsdUnreal::MaterialUtils::UnregisterRenderContext(UnrealIdentifiers::MdlRenderContext);
+		UsdUnreal::MaterialUtils::UnregisterRenderContext(UnrealIdentifiers::MaterialXRenderContext);
 #endif	  // WITH_EDITOR
 	}
 
@@ -116,14 +114,19 @@ public:
 		return UsdSchemaTranslatorRegistry;
 	}
 
+	PRAGMA_DISABLE_DEPRECATION_WARNINGS
 	virtual FUsdRenderContextRegistry& GetRenderContextRegistry() override
 	{
 		return UsdRenderContextRegistry;
 	}
+	PRAGMA_ENABLE_DEPRECATION_WARNINGS
 
 protected:
 	FUsdSchemaTranslatorRegistry UsdSchemaTranslatorRegistry;
+
+	PRAGMA_DISABLE_DEPRECATION_WARNINGS
 	FUsdRenderContextRegistry UsdRenderContextRegistry;
+	PRAGMA_ENABLE_DEPRECATION_WARNINGS
 
 	TArray<FRegisteredSchemaTranslatorHandle> TranslatorHandles;
 };
diff --git a/Engine/Plugins/Importers/USDImporter/Source/USDSchemas/Public/USDSchemaTranslator.h b/Engine/Plugins/Importers/USDImporter/Source/USDSchemas/Public/USDSchemaTranslator.h
--- a/Engine/Plugins/Importers/USDImporter/Source/USDSchemas/Public/USDSchemaTranslator.h
+++ b/Engine/Plugins/Importers/USDImporter/Source/USDSchemas/Public/USDSchemaTranslator.h
@@ -140,37 +140,16 @@ public:
 	FCreateTranslator CreateFunction;
 };
 
-class USDSCHEMAS_API FUsdRenderContextRegistry
+class UE_DEPRECATED(5.5, "Use the render context functions in USDMaterialUtils.h instead.") USDSCHEMAS_API FUsdRenderContextRegistry
 {
 public:
 	FUsdRenderContextRegistry();
 
-	void Register(const FName& RenderContextToken)
-	{
-		RegisteredRenderContexts.Add(RenderContextToken);
-	}
-	void Unregister(const FName& RenderContextToken)
-	{
-		RegisteredRenderContexts.Remove(RenderContextToken);
-	}
-
-	const TSet<FName>& GetRenderContexts() const
-	{
-		return RegisteredRenderContexts;
-	}
-	const FName& GetUniversalRenderContext() const
-	{
-		return UniversalRenderContext;
-	}
-	const FName& GetUnrealRenderContext() const
-	{
-		return UnrealRenderContext;
-	}
-
-protected:
-	TSet<FName> RegisteredRenderContexts;
-	FName UniversalRenderContext;
-	FName UnrealRenderContext;
+	void Register(const FName& RenderContextToken);
+	void Unregister(const FName& RenderContextToken);
+	const TSet<FName>& GetRenderContexts() const;
+	const FName& GetUniversalRenderContext() const;
+	const FName& GetUnrealRenderContext() const;
 };
 
 struct USDSCHEMAS_API FUsdSchemaTranslationContext : public TSharedFromThis<FUsdSchemaTranslationContext>
diff --git a/Engine/Plugins/Importers/USDImporter/Source/USDSchemas/Public/USDSchemasModule.h b/Engine/Plugins/Importers/USDImporter/Source/USDSchemas/Public/USDSchemasModule.h
--- a/Engine/Plugins/Importers/USDImporter/Source/USDSchemas/Public/USDSchemasModule.h
+++ b/Engine/Plugins/Importers/USDImporter/Source/USDSchemas/Public/USDSchemasModule.h
@@ -14,7 +14,10 @@ class IUsdSchemasModule : public IModuleInterface
 {
 public:
 	virtual FUsdSchemaTranslatorRegistry& GetTranslatorRegistry() = 0;
+
+	PRAGMA_DISABLE_DEPRECATION_WARNINGS
 	virtual FUsdRenderContextRegistry& GetRenderContextRegistry() = 0;
+	PRAGMA_ENABLE_DEPRECATION_WARNINGS
 };
 
 namespace UsdUnreal::Analytics
diff --git a/Engine/Plugins/Importers/USDImporter/Source/USDStage/Private/USDStageActor.cpp b/Engine/Plugins/Importers/USDImporter/Source/USDStage/Private/USDStageActor.cpp
--- a/Engine/Plugins/Importers/USDImporter/Source/USDStage/Private/USDStageActor.cpp
+++ b/Engine/Plugins/Importers/USDImporter/Source/USDStage/Private/USDStageActor.cpp
@@ -925,6 +925,7 @@ AUsdStageActor::AUsdStageActor()
 	, bShareAssetsForIdenticalPrims(true)
 	, PurposesToLoad((int32)EUsdPurpose::Proxy)
 	, NaniteTriangleThreshold((uint64)1000000)
+	, RenderContext(UnrealIdentifiers::UnrealRenderContext)
 	, MaterialPurpose(*UnrealIdentifiers::MaterialPreviewPurpose)
 	, RootMotionHandling(EUsdRootMotionHandling::NoAdditionalRootMotion)
 	, SubdivisionLevel(0)
@@ -952,9 +953,6 @@ AUsdStageActor::AUsdStageActor()
 	// some transitions like reinstantiation.
 	// c.f. doc comment on FRecompilationTracker for more info.
 
-	IUsdSchemasModule& UsdSchemasModule = FModuleManager::Get().LoadModuleChecked<IUsdSchemasModule>(TEXT("USDSchemas"));
-	RenderContext = UsdSchemasModule.GetRenderContextRegistry().GetUnrealRenderContext();
-
 	// This is marked as a subobject now or else instances of blueprints that derive the AUsdStageActor don't get new transactors created
 	// for them
 	const FName UniqueName = MakeUniqueObjectName(this, UUsdTransactor::StaticClass(), TEXT("Transactor"));
diff --git a/Engine/Plugins/Importers/USDImporter/Source/USDStage/Private/USDStageActorCustomization.cpp b/Engine/Plugins/Importers/USDImporter/Source/USDStage/Private/USDStageActorCustomization.cpp
--- a/Engine/Plugins/Importers/USDImporter/Source/USDStage/Private/USDStageActorCustomization.cpp
+++ b/Engine/Plugins/Importers/USDImporter/Source/USDStage/Private/USDStageActorCustomization.cpp
@@ -4,8 +4,8 @@
 
 #include "USDStageActorCustomization.h"
 
+#include "USDMaterialUtils.h"
 #include "USDProjectSettings.h"
-#include "USDSchemasModule.h"
 #include "USDSchemaTranslator.h"
 #include "USDStageActor.h"
 
@@ -53,16 +53,14 @@ void FUsdStageActorCustomization::CustomizeDetails(IDetailLayoutBuilder& DetailL
 
 	CurrentActor->OnStageChanged.AddSP(this, &FUsdStageActorCustomization::ForceRefreshDetails);
 
-	IUsdSchemasModule& UsdSchemasModule = FModuleManager::Get().LoadModuleChecked<IUsdSchemasModule>(TEXT("USDSchemas"));
-
 	RenderContextComboBoxItems.Reset();
 	TSharedPtr<FString> InitiallySelectedContext;
-	for (const FName& Context : UsdSchemasModule.GetRenderContextRegistry().GetRenderContexts())
+	for (const FName& Context : UsdUnreal::MaterialUtils::GetRegisteredRenderContexts())
 	{
 		TSharedPtr<FString> ContextStr;
-		if (Context == NAME_None)
+		if (Context == UnrealIdentifiers::UniversalRenderContext)
 		{
-			ContextStr = MakeShared<FString>(TEXT("universal"));
+			ContextStr = MakeShared<FString>(UnrealIdentifiers::UniversalRenderContextDisplayString);
 		}
 		else
 		{
@@ -291,7 +289,8 @@ void FUsdStageActorCustomization::OnComboBoxSelectionChanged(TSharedPtr<FString>
 		FText::FromString(NewContext.IsValid() ? *NewContext : TEXT("None"))
 	));
 
-	FName NewContextName = (*NewContext) == TEXT("universal") ? NAME_None : FName(**NewContext);
+	FName NewContextName = (*NewContext) == UnrealIdentifiers::UniversalRenderContextDisplayString ? UnrealIdentifiers::UniversalRenderContext
+																								   : FName(**NewContext);
 
 	CurrentActor->SetRenderContext(NewContextName);
 }
diff --git a/Engine/Plugins/Importers/USDImporter/Source/USDStageEditor/Private/SUSDStage.cpp b/Engine/Plugins/Importers/USDImporter/Source/USDStageEditor/Private/SUSDStage.cpp
--- a/Engine/Plugins/Importers/USDImporter/Source/USDStageEditor/Private/SUSDStage.cpp
+++ b/Engine/Plugins/Importers/USDImporter/Source/USDStageEditor/Private/SUSDStage.cpp
@@ -12,8 +12,8 @@
 #include "USDErrorUtils.h"
 #include "USDLayerUtils.h"
 #include "USDLog.h"
+#include "USDMaterialUtils.h"
 #include "USDProjectSettings.h"
-#include "USDSchemasModule.h"
 #include "USDSchemaTranslator.h"
 #include "USDStageActor.h"
 #include "USDStageEditorSettings.h"
@@ -1437,21 +1437,21 @@ void SUsdStage::FillPurposesToLoadSubMenu(FMenuBuilder& MenuBuilder)
 
 void SUsdStage::FillRenderContextSubMenu(FMenuBuilder& MenuBuilder)
 {
-	auto AddRenderContextEntry = [&](const FName& RenderContext)
+	auto AddRenderContextEntry = [&](const FName& RenderContextName)
 	{
-		FText RenderContextName = FText::FromName(RenderContext);
-		if (RenderContext.IsNone())
+		FText RenderContextText = FText::FromName(RenderContextName);
+		if (RenderContextName == UnrealIdentifiers::UniversalRenderContext)
 		{
-			RenderContextName = LOCTEXT("UniversalRenderContext", "universal");
+			RenderContextText = FText::FromString(UnrealIdentifiers::UniversalRenderContextDisplayString);
 		}
 
 		MenuBuilder.AddMenuEntry(
-			RenderContextName,
+			RenderContextText,
 			FText::GetEmpty(),
 			FSlateIcon(),
 			FUIAction(
 				FExecuteAction::CreateLambda(
-					[this, RenderContext]()
+					[this, RenderContextName]()
 					{
 						if (AUsdStageActor* StageActor = GetStageActorOrCDO())
 						{
@@ -1463,7 +1463,7 @@ void SUsdStage::FillRenderContextSubMenu(FMenuBuilder& MenuBuilder)
 							// c.f. comment in SUsdStage::FillCollapsingSubMenu
 							TGuardValue<bool> MaintainSelectionGuard(bUpdatingViewportSelection, true);
 
-							StageActor->SetRenderContext(RenderContext);
+							StageActor->SetRenderContext(RenderContextName);
 							if (StageActor->IsTemplate())
 							{
 								StageActor->SaveConfig();
@@ -1473,11 +1473,11 @@ void SUsdStage::FillRenderContextSubMenu(FMenuBuilder& MenuBuilder)
 				),
 				FCanExecuteAction{},
 				FIsActionChecked::CreateLambda(
-					[this, RenderContext]()
+					[this, RenderContextName]()
 					{
 						if (AUsdStageActor* StageActor = GetStageActorOrCDO())
 						{
-							return StageActor->RenderContext == RenderContext;
+							return StageActor->RenderContext == RenderContextName;
 						}
 						return false;
 					}
@@ -1488,9 +1488,7 @@ void SUsdStage::FillRenderContextSubMenu(FMenuBuilder& MenuBuilder)
 		);
 	};
 
-	IUsdSchemasModule& UsdSchemasModule = FModuleManager::Get().LoadModuleChecked<IUsdSchemasModule>(TEXT("USDSchemas"));
-
-	for (const FName& RenderContext : UsdSchemasModule.GetRenderContextRegistry().GetRenderContexts())
+	for (const FName& RenderContext : UsdUnreal::MaterialUtils::GetRegisteredRenderContexts())
 	{
 		AddRenderContextEntry(RenderContext);
 	}
diff --git a/Engine/Plugins/Importers/USDImporter/Source/USDStageImporter/Private/USDStageImportOptions.cpp b/Engine/Plugins/Importers/USDImporter/Source/USDStageImporter/Private/USDStageImportOptions.cpp
--- a/Engine/Plugins/Importers/USDImporter/Source/USDStageImporter/Private/USDStageImportOptions.cpp
+++ b/Engine/Plugins/Importers/USDImporter/Source/USDStageImporter/Private/USDStageImportOptions.cpp
@@ -27,7 +27,7 @@ UUsdStageImportOptions::UUsdStageImportOptions(const FObjectInitializer& ObjectI
 	PurposesToImport = (int32)(EUsdPurpose::Default | EUsdPurpose::Proxy | EUsdPurpose::Render | EUsdPurpose::Guide);
 	NaniteTriangleThreshold = INT32_MAX;
 	IUsdSchemasModule& UsdSchemasModule = FModuleManager::Get().LoadModuleChecked<IUsdSchemasModule>(TEXT("USDSchemas"));
-	RenderContextToImport = UsdSchemasModule.GetRenderContextRegistry().GetUnrealRenderContext();
+	RenderContextToImport = UnrealIdentifiers::UnrealRenderContext;
 	MaterialPurpose = *UnrealIdentifiers::MaterialPreviewPurpose;
 	SubdivisionLevel = 0;
 	MetadataOptions = FUsdMetadataImportOptions{
diff --git a/Engine/Plugins/Importers/USDImporter/Source/USDStageImporter/Private/USDStageImportOptionsCustomization.cpp b/Engine/Plugins/Importers/USDImporter/Source/USDStageImporter/Private/USDStageImportOptionsCustomization.cpp
--- a/Engine/Plugins/Importers/USDImporter/Source/USDStageImporter/Private/USDStageImportOptionsCustomization.cpp
+++ b/Engine/Plugins/Importers/USDImporter/Source/USDStageImporter/Private/USDStageImportOptionsCustomization.cpp
@@ -2,8 +2,8 @@
 
 #include "USDStageImportOptionsCustomization.h"
 
+#include "USDMaterialUtils.h"
 #include "USDProjectSettings.h"
-#include "USDSchemasModule.h"
 #include "USDSchemaTranslator.h"
 #include "USDStageImportOptions.h"
 
@@ -56,16 +56,14 @@ void FUsdStageImportOptionsCustomization::CustomizeDetails(IDetailLayoutBuilder&
 		DetailLayoutBuilder.HideProperty(PrimsToImportProperty);
 	}
 
-	IUsdSchemasModule& UsdSchemasModule = FModuleManager::Get().LoadModuleChecked<IUsdSchemasModule>(TEXT("USDSchemas"));
-
 	RenderContextComboBoxItems.Reset();
 	TSharedPtr<FString> InitiallySelectedContext;
-	for (const FName& Context : UsdSchemasModule.GetRenderContextRegistry().GetRenderContexts())
+	for (const FName& Context : UsdUnreal::MaterialUtils::GetRegisteredRenderContexts())
 	{
 		TSharedPtr<FString> ContextStr;
-		if (Context == NAME_None)
+		if (Context == UnrealIdentifiers::UniversalRenderContext)
 		{
-			ContextStr = MakeShared<FString>(TEXT("universal"));
+			ContextStr = MakeShared<FString>(UnrealIdentifiers::UniversalRenderContextDisplayString);
 		}
 		else
 		{
@@ -276,7 +274,8 @@ void FUsdStageImportOptionsCustomization::OnComboBoxSelectionChanged(TSharedPtr<
 		return;
 	}
 
-	FName NewContextName = (*NewContext) == TEXT("universal") ? NAME_None : FName(**NewContext);
+	FName NewContextName = (*NewContext) == UnrealIdentifiers::UniversalRenderContextDisplayString ? UnrealIdentifiers::UniversalRenderContext
+																								   : FName(**NewContext);
 
 	CurrentOptions->RenderContextToImport = NewContextName;
 }
diff --git a/Engine/Plugins/Interchange/Editor/InterchangeEditor.uplugin b/Engine/Plugins/Interchange/Editor/InterchangeEditor.uplugin
--- a/Engine/Plugins/Interchange/Editor/InterchangeEditor.uplugin
+++ b/Engine/Plugins/Interchange/Editor/InterchangeEditor.uplugin
@@ -56,6 +56,10 @@
 				"Editor",
 				"Program"
 			]
+		},
+		{
+			"Name": "USDCore",
+			"Enabled": true
 		}
 	],
 	"IsExperimentalVersion": false
diff --git a/Engine/Plugins/Interchange/Editor/Source/InterchangeEditor/InterchangeEditor.Build.cs b/Engine/Plugins/Interchange/Editor/Source/InterchangeEditor/InterchangeEditor.Build.cs
--- a/Engine/Plugins/Interchange/Editor/Source/InterchangeEditor/InterchangeEditor.Build.cs
+++ b/Engine/Plugins/Interchange/Editor/Source/InterchangeEditor/InterchangeEditor.Build.cs
@@ -34,7 +34,11 @@ namespace UnrealBuildTool.Rules
 					"InterchangePipelines",
 					"MessageLog",
 					"SlateCore",
+					"InputCore", // Translator settings customizations
+					"Slate", // Translator settings customizations
 					"ToolMenus",
+					"UnrealUSDWrapper", // UnrealIdentifiers in the USD translator settings customization
+					"USDClasses", // USDProjectSettings, to fetch other material purporses
 				}
 			);
 		}
diff --git a/Engine/Plugins/Interchange/Editor/Source/InterchangeEditor/Private/InterchangeEditorModule.cpp b/Engine/Plugins/Interchange/Editor/Source/InterchangeEditor/Private/InterchangeEditorModule.cpp
--- a/Engine/Plugins/Interchange/Editor/Source/InterchangeEditor/Private/InterchangeEditorModule.cpp
+++ b/Engine/Plugins/Interchange/Editor/Source/InterchangeEditor/Private/InterchangeEditorModule.cpp
@@ -7,10 +7,12 @@
 #include "InterchangeEditorLog.h"
 #include "InterchangeFbxAssetImportDataConverter.h"
 #include "InterchangeManager.h"
+#include "InterchangeUsdTranslatorSettingsCustomization.h"
 #include "MessageLogModule.h"
 #include "Misc/App.h"
 #include "Misc/CoreDelegates.h"
 #include "Modules/ModuleManager.h"
+#include "PropertyEditorModule.h"
 #include "Settings/EditorLoadingSavingSettings.h"
 
 #define LOCTEXT_NAMESPACE "InterchangeEditorModule"
@@ -131,6 +133,13 @@ void FInterchangeEditorModule::StartupModule()
 		};
 
 		InterchangeManager.OnPreDestroyInterchangeManager.AddLambda(UnregisterItems);
+
+		// Translator settings customizations
+		FPropertyEditorModule& PropertyModule = FModuleManager::LoadModuleChecked<FPropertyEditorModule>(TEXT("PropertyEditor"));
+		PropertyModule.RegisterCustomClassLayout(
+			TEXT("InterchangeUsdTranslatorSettings"),
+			FOnGetDetailCustomizationInstance::CreateStatic(&FInterchangeUsdTranslatorSettingsCustomization::MakeInstance)
+		);
 	};
 
 	if (GEngine)
@@ -143,6 +152,15 @@ void FInterchangeEditorModule::StartupModule()
 	}
 }
 
+void FInterchangeEditorModule::ShutdownModule()
+{
+	//Translator settings customizations
+	if (FPropertyEditorModule* PropertyModule = FModuleManager::GetModulePtr<FPropertyEditorModule>(TEXT("PropertyEditor")))
+	{
+		PropertyModule->UnregisterCustomClassLayout(TEXT("InterchangeUsdTranslatorSettings"));
+	}
+}
+
 IMPLEMENT_MODULE(FInterchangeEditorModule, InterchangeEditor)
 
 #undef LOCTEXT_NAMESPACE // "InterchangeEditorModule"
diff --git a/Engine/Plugins/Interchange/Editor/Source/InterchangeEditor/Private/InterchangeUsdTranslatorSettingsCustomization.cpp b/Engine/Plugins/Interchange/Editor/Source/InterchangeEditor/Private/InterchangeUsdTranslatorSettingsCustomization.cpp
new file mode 100644
--- /dev/null
+++ b/Engine/Plugins/Interchange/Editor/Source/InterchangeEditor/Private/InterchangeUsdTranslatorSettingsCustomization.cpp
@@ -0,0 +1,302 @@
+// Copyright Epic Games, Inc. All Rights Reserved.
+
+#include "InterchangeUsdTranslatorSettingsCustomization.h"
+
+#include "UnrealUSDWrapper.h"
+#include "Usd/InterchangeUsdTranslator.h"
+#include "USDMaterialUtils.h"
+#include "USDProjectSettings.h"
+
+#include "DetailCategoryBuilder.h"
+#include "DetailLayoutBuilder.h"
+#include "DetailWidgetRow.h"
+#include "Modules/ModuleManager.h"
+#include "ScopedTransaction.h"
+#include "Styling/AppStyle.h"
+#include "UObject/Object.h"
+#include "Widgets/DeclarativeSyntaxSupport.h"
+#include "Widgets/Input/SEditableTextBox.h"
+
+#define LOCTEXT_NAMESPACE "InterchangeUsdTranslatorSettingsCustomization"
+
+FInterchangeUsdTranslatorSettingsCustomization::FInterchangeUsdTranslatorSettingsCustomization()
+{
+}
+
+TSharedRef<IDetailCustomization> FInterchangeUsdTranslatorSettingsCustomization::MakeInstance()
+{
+	return MakeShared<FInterchangeUsdTranslatorSettingsCustomization>();
+}
+
+void FInterchangeUsdTranslatorSettingsCustomization::CustomizeDetails(IDetailLayoutBuilder& DetailLayoutBuilder)
+{
+	TArray<TWeakObjectPtr<UObject>> SelectedObjects = DetailLayoutBuilder.GetSelectedObjects();
+	if (SelectedObjects.Num() != 1)
+	{
+		return;
+	}
+
+	TWeakObjectPtr<UObject> SelectedObject = SelectedObjects[0];
+	if (!SelectedObject.IsValid())
+	{
+		return;
+	}
+
+	CurrentOptions = Cast<UInterchangeUsdTranslatorSettings>(SelectedObject.Get());
+	if (!CurrentOptions)
+	{
+		return;
+	}
+
+	RenderContextComboBoxItems.Reset();
+	TSharedPtr<FString> InitiallySelectedContext;
+	for (const FName& Context : UsdUnreal::MaterialUtils::GetRegisteredRenderContexts())
+	{
+		TSharedPtr<FString> ContextStr;
+		if (Context == UnrealIdentifiers::UniversalRenderContext)
+		{
+			ContextStr = MakeShared<FString>(UnrealIdentifiers::UniversalRenderContextDisplayString);
+		}
+		else
+		{
+			ContextStr = MakeShared<FString>(Context.ToString());
+		}
+
+		if (Context == CurrentOptions->RenderContext)
+		{
+			InitiallySelectedContext = ContextStr;
+		}
+
+		RenderContextComboBoxItems.Add(ContextStr);
+	}
+
+	IDetailCategoryBuilder& CatBuilder = DetailLayoutBuilder.EditCategory(TEXT("USD Translator"));
+
+	if (TSharedPtr<IPropertyHandle> RenderContextProperty = DetailLayoutBuilder.GetProperty(
+			GET_MEMBER_NAME_CHECKED(UInterchangeUsdTranslatorSettings, RenderContext)
+		))
+	{
+		DetailLayoutBuilder.HideProperty(RenderContextProperty);
+
+		// clang-format off
+		CatBuilder.AddCustomRow(FText::FromString(TEXT("RenderContextCustomization")))
+		.NameContent()
+		[
+			SNew(STextBlock)
+			.Text(FText::FromString(TEXT("Render Context")))
+			.Font(FAppStyle::GetFontStyle(TEXT("PropertyWindow.NormalFont")))
+			.ToolTipText(RenderContextProperty->GetToolTipText())
+		]
+		.ValueContent()
+		[
+			SAssignNew(RenderContextComboBox, SComboBox<TSharedPtr<FString>>)
+			.OptionsSource(&RenderContextComboBoxItems)
+			.InitiallySelectedItem(InitiallySelectedContext)
+			.OnSelectionChanged_Lambda([this, RenderContextProperty](TSharedPtr<FString> NewContext, ESelectInfo::Type SelectType)
+			{
+				if (CurrentOptions == nullptr || !NewContext.IsValid())
+				{
+					return;
+				}
+
+				FScopedTransaction Transaction(LOCTEXT("RenderContextTransaction", "Edit Render Context"));
+				RenderContextProperty->NotifyPreChange();
+				{
+					FName NewContextName = (*NewContext) == UnrealIdentifiers::UniversalRenderContextDisplayString
+						? UnrealIdentifiers::UniversalRenderContext
+						: FName(**NewContext);
+
+					CurrentOptions->RenderContext = NewContextName;
+				}
+				RenderContextProperty->NotifyPostChange(EPropertyChangeType::ValueSet);
+				RenderContextProperty->NotifyFinishedChangingProperties();
+			})
+			.OnGenerateWidget_Lambda([](TSharedPtr<FString> Item)
+			{
+				return SNew(STextBlock)
+					.Text(Item.IsValid() ? FText::FromString(*Item) : FText::GetEmpty())
+					.Font(FAppStyle::GetFontStyle(TEXT("PropertyWindow.NormalFont")));
+			})
+			.Content()
+			[
+				SNew(STextBlock)
+				.Text_Lambda([this]() -> FText
+				{
+					TSharedPtr<FString> SelectedItem = RenderContextComboBox->GetSelectedItem();
+					if (SelectedItem.IsValid())
+					{
+						return FText::FromString(*SelectedItem);
+					}
+
+					return FText::GetEmpty();
+				})
+				.Font(FAppStyle::GetFontStyle(TEXT("PropertyWindow.NormalFont")))
+			]
+		];
+		// clang-format on
+	}
+
+	if (TSharedPtr<IPropertyHandle> MaterialPurposeProperty = DetailLayoutBuilder.GetProperty(
+			GET_MEMBER_NAME_CHECKED(UInterchangeUsdTranslatorSettings, MaterialPurpose)
+		))
+	{
+		DetailLayoutBuilder.HideProperty(MaterialPurposeProperty);
+
+		// clang-format off
+		CatBuilder.AddCustomRow(FText::FromString(TEXT("MaterialPurposeCustomization")))
+		.NameContent()
+		[
+			SNew(STextBlock)
+			.Text(FText::FromString(TEXT("Material purpose")))
+			.Font(FAppStyle::GetFontStyle(TEXT("PropertyWindow.NormalFont")))
+			.ToolTipText(MaterialPurposeProperty->GetToolTipText())
+		]
+		.ValueContent()
+		[
+			SNew(SBox)
+			.VAlign(VAlign_Center)
+			[
+				SNew(SComboBox<TSharedPtr<FString>>)
+				.OptionsSource(&MaterialPurposeComboBoxItems)
+				.OnComboBoxOpening_Lambda([this]()
+				{
+					MaterialPurposeComboBoxItems = {
+						MakeShared<FString>(UnrealIdentifiers::MaterialAllPurpose),
+						MakeShared<FString>(UnrealIdentifiers::MaterialPreviewPurpose),
+						MakeShared<FString>(UnrealIdentifiers::MaterialFullPurpose)
+					};
+
+					// Add additional purposes from project settings
+					if (const UUsdProjectSettings* ProjectSettings = GetDefault<UUsdProjectSettings>())
+					{
+						MaterialPurposeComboBoxItems.Reserve(MaterialPurposeComboBoxItems.Num() + ProjectSettings->AdditionalMaterialPurposes.Num());
+
+						TSet<FString> ExistingEntries = {
+							UnrealIdentifiers::MaterialAllPurpose,
+							UnrealIdentifiers::MaterialPreviewPurpose,
+							UnrealIdentifiers::MaterialFullPurpose
+						};
+
+						for (const FName& AdditionalPurpose : ProjectSettings->AdditionalMaterialPurposes)
+						{
+							FString AdditionalPurposeStr = AdditionalPurpose.ToString();
+
+							if (!ExistingEntries.Contains(AdditionalPurposeStr))
+							{
+								ExistingEntries.Add(AdditionalPurposeStr);
+								MaterialPurposeComboBoxItems.AddUnique(MakeShared<FString>(AdditionalPurposeStr));
+							}
+						}
+					}
+				})
+				.OnGenerateWidget_Lambda([&](TSharedPtr<FString> Option)
+				{
+					TSharedPtr<SWidget> Widget = SNullWidget::NullWidget;
+					if (Option)
+					{
+						Widget = SNew(STextBlock)
+							.Text(FText::FromString((*Option) == UnrealIdentifiers::MaterialAllPurpose
+								? UnrealIdentifiers::MaterialAllPurposeText
+								: *Option
+							))
+							.Font(FAppStyle::GetFontStyle("PropertyWindow.NormalFont"));
+					}
+
+					return Widget.ToSharedRef();
+				})
+				.OnSelectionChanged_Lambda([this, MaterialPurposeProperty](TSharedPtr<FString> ChosenOption, ESelectInfo::Type SelectInfo)
+				{
+					if (CurrentOptions && ChosenOption)
+					{
+						FScopedTransaction Transaction(LOCTEXT("MaterialPurposeTransaction", "Edit Material Purpose"));
+						MaterialPurposeProperty->NotifyPreChange();
+						{
+							CurrentOptions->MaterialPurpose = **ChosenOption;
+						}
+						MaterialPurposeProperty->NotifyPostChange(EPropertyChangeType::ValueSet);
+						MaterialPurposeProperty->NotifyFinishedChangingProperties();
+					}
+
+				})
+				[
+					SNew(SEditableTextBox)
+					.Text_Lambda([this]() -> FText
+					{
+						if (CurrentOptions)
+						{
+							return FText::FromString(CurrentOptions->MaterialPurpose == *UnrealIdentifiers::MaterialAllPurpose
+								? UnrealIdentifiers::MaterialAllPurposeText
+								: CurrentOptions->MaterialPurpose.ToString()
+							);
+						}
+
+						return FText::GetEmpty();
+					})
+					.Font(FAppStyle::GetFontStyle("PropertyWindow.NormalFont"))
+					.OnTextCommitted_Lambda([this, MaterialPurposeProperty](const FText& NewText, ETextCommit::Type CommitType)
+					{
+						if (CommitType != ETextCommit::OnEnter)
+						{
+							return;
+						}
+
+						FString NewPurposeString = NewText.ToString();
+						FName NewPurpose = *NewPurposeString;
+
+						bool bIsNew = true;
+						for (const TSharedPtr<FString>& Purpose : MaterialPurposeComboBoxItems)
+						{
+							if (Purpose && *Purpose == NewPurposeString)
+							{
+								bIsNew = false;
+								break;
+							}
+						}
+
+						if (bIsNew)
+						{
+							if (UUsdProjectSettings* ProjectSettings = GetMutableDefault<UUsdProjectSettings>())
+							{
+								ProjectSettings->AdditionalMaterialPurposes.AddUnique(NewPurpose);
+								ProjectSettings->SaveConfig();
+							}
+						}
+
+						if (CurrentOptions)
+						{
+							FScopedTransaction Transaction(LOCTEXT("MaterialPurposeTypeTransaction", "Add and Set Material Purpose"));
+							MaterialPurposeProperty->NotifyPreChange();
+							{
+								CurrentOptions->MaterialPurpose = NewPurpose;
+							}
+							MaterialPurposeProperty->NotifyPostChange(EPropertyChangeType::ValueSet);
+							MaterialPurposeProperty->NotifyFinishedChangingProperties();
+						}
+					})
+				]
+			]
+		];
+		// clang-format on
+	}
+
+	// Add/remove properties so that they retain their usual order
+	if (TSharedPtr<IPropertyHandle> OverrideStageOptionsProperty = DetailLayoutBuilder.GetProperty(
+			GET_MEMBER_NAME_CHECKED(UInterchangeUsdTranslatorSettings, bOverrideStageOptions)
+		))
+	{
+		CatBuilder.AddProperty(OverrideStageOptionsProperty);
+	}
+	if (TSharedPtr<IPropertyHandle> StageOptionsProperty = DetailLayoutBuilder.GetProperty(
+			GET_MEMBER_NAME_CHECKED(UInterchangeUsdTranslatorSettings, StageOptions)
+		))
+	{
+		CatBuilder.AddProperty(StageOptionsProperty);
+	}
+}
+
+void FInterchangeUsdTranslatorSettingsCustomization::CustomizeDetails(const TSharedPtr<IDetailLayoutBuilder>& DetailBuilder)
+{
+	CustomizeDetails(*DetailBuilder);
+}
+
+#undef LOCTEXT_NAMESPACE
diff --git a/Engine/Plugins/Interchange/Editor/Source/InterchangeEditor/Private/InterchangeUsdTranslatorSettingsCustomization.h b/Engine/Plugins/Interchange/Editor/Source/InterchangeEditor/Private/InterchangeUsdTranslatorSettingsCustomization.h
new file mode 100644
--- /dev/null
+++ b/Engine/Plugins/Interchange/Editor/Source/InterchangeEditor/Private/InterchangeUsdTranslatorSettingsCustomization.h
@@ -0,0 +1,33 @@
+// Copyright Epic Games, Inc. All Rights Reserved.
+
+#pragma once
+
+#if WITH_EDITOR
+
+#include "IDetailCustomization.h"
+#include "Widgets/Input/SComboBox.h"
+
+class UInterchangeUsdTranslatorSettings;
+class IDetailLayoutBuilder;
+
+class FInterchangeUsdTranslatorSettingsCustomization : public IDetailCustomization
+{
+public:
+	FInterchangeUsdTranslatorSettingsCustomization();
+	static TSharedRef<IDetailCustomization> MakeInstance();
+
+	// IDetailCustomization interface
+	virtual void CustomizeDetails(IDetailLayoutBuilder& DetailLayout) override;
+	virtual void CustomizeDetails(const TSharedPtr<IDetailLayoutBuilder>& DetailBuilder) override;
+	// End of IDetailCustomization interface
+
+private:
+	UInterchangeUsdTranslatorSettings* CurrentOptions;
+
+	TSharedPtr<SComboBox<TSharedPtr<FString>>> RenderContextComboBox;
+	TArray<TSharedPtr<FString>> RenderContextComboBoxItems;
+
+	TArray<TSharedPtr<FString>> MaterialPurposeComboBoxItems;
+};
+
+#endif	  // WITH_EDITOR
diff --git a/Engine/Plugins/Interchange/Editor/Source/InterchangeEditor/Public/InterchangeEditorModule.h b/Engine/Plugins/Interchange/Editor/Source/InterchangeEditor/Public/InterchangeEditorModule.h
--- a/Engine/Plugins/Interchange/Editor/Source/InterchangeEditor/Public/InterchangeEditorModule.h
+++ b/Engine/Plugins/Interchange/Editor/Source/InterchangeEditor/Public/InterchangeEditorModule.h
@@ -16,4 +16,5 @@ public:
 
 private:
 	virtual void StartupModule() override;
+	virtual void ShutdownModule() override;
 };
diff --git a/Engine/Plugins/Interchange/Runtime/Source/Import/Private/Usd/InterchangeUsdTranslator.cpp b/Engine/Plugins/Interchange/Runtime/Source/Import/Private/Usd/InterchangeUsdTranslator.cpp
--- a/Engine/Plugins/Interchange/Runtime/Source/Import/Private/Usd/InterchangeUsdTranslator.cpp
+++ b/Engine/Plugins/Interchange/Runtime/Source/Import/Private/Usd/InterchangeUsdTranslator.cpp
@@ -738,7 +738,7 @@ namespace UE::InterchangeUsdTranslator::Private
 			}
 		};
 
-		FString RenderContext = TranslatorSettings ? TranslatorSettings->RenderContext.ToString() : FString();
+		FName RenderContext = TranslatorSettings ? TranslatorSettings->RenderContext : UnrealIdentifiers::UniversalRenderContext;
 
 		// Check for any references of MaterialX
 #if WITH_EDITOR
@@ -790,8 +790,7 @@ namespace UE::InterchangeUsdTranslator::Private
 		}
 #endif	  // WITH_EDITOR
 
-		const static FName UnrealName = *UsdToUnreal::ConvertToken(UnrealIdentifiers::Unreal);
-		if (RenderContext == UnrealName)
+		if (RenderContext == UnrealIdentifiers::UnrealRenderContext)
 		{
 			UE_LOG(
 				LogUsd,
@@ -801,7 +800,7 @@ namespace UE::InterchangeUsdTranslator::Private
 				),
 				*PrimPath
 			);
-			RenderContext = FString{};
+			RenderContext = UnrealIdentifiers::UniversalRenderContext;
 		}
 
 		UInterchangeMaterialInstanceNode* MaterialNode = NewObject<UInterchangeMaterialInstanceNode>(&NodeContainer);
@@ -814,7 +813,7 @@ namespace UE::InterchangeUsdTranslator::Private
 		SetMaterialSlotDependencies();
 
 		UsdToUnreal::FUsdPreviewSurfaceMaterialData MaterialData;
-		const bool bSuccess = UsdToUnreal::ConvertMaterial(Prim, MaterialData, TranslatorSettings ? *RenderContext : nullptr);
+		const bool bSuccess = UsdToUnreal::ConvertMaterial(Prim, MaterialData, *RenderContext.ToString());
 
 		// Set all the parameter values to the interchange node
 		bool bHasUDIMTexture = false;
@@ -2335,7 +2334,7 @@ namespace UE::InterchangeUsdTranslator::Private
 
 UInterchangeUsdTranslatorSettings::UInterchangeUsdTranslatorSettings()
 	: GeometryPurpose((int32)(EUsdPurpose::Default | EUsdPurpose::Proxy | EUsdPurpose::Render | EUsdPurpose::Guide))
-	, RenderContext(NAME_None)	  // Default to the universal render context for now as we don't support 'unreal' yet anyway
+	, RenderContext(UnrealIdentifiers::UniversalRenderContext)	  // Default to the universal render context for now as we don't support 'unreal' yet
 	, MaterialPurpose(*UnrealIdentifiers::MaterialPreviewPurpose)
 	, InterpolationType(EUsdInterpolationType::Linear)
 	, bOverrideStageOptions(false)
@@ -2428,8 +2427,9 @@ bool UInterchangeUSDTranslator::Translate(UInterchangeBaseNodeContainer& NodeCon
 	// really do anything anyway as the module doesn't use IMPLEMENT_MODULE_USD!
 	// Luckily we can get around this here because pxr::TfToken doesn't allocate on its own: At most USD makes a copy of the string, which it should
 	// allocate/deallocate on its own allocator.
-	MeshOptions.RenderContext = Settings->RenderContext.IsNone() ? pxr::UsdShadeTokens->universalRenderContext
-																 : pxr::TfToken{TCHAR_TO_ANSI(*Settings->RenderContext.ToString())};
+	MeshOptions.RenderContext = Settings->RenderContext == UnrealIdentifiers::UniversalRenderContext
+									? pxr::UsdShadeTokens->universalRenderContext
+									: pxr::TfToken{TCHAR_TO_ANSI(*Settings->RenderContext.ToString())};
 	MeshOptions.MaterialPurpose = Settings->MaterialPurpose.IsNone() ? pxr::UsdShadeTokens->allPurpose
 																	 : pxr::TfToken{TCHAR_TO_ANSI(*Settings->MaterialPurpose.ToString())};
 
diff --git a/Engine/Plugins/Runtime/USDCore/Source/USDClasses/Private/USDMaterialUtils.cpp b/Engine/Plugins/Runtime/USDCore/Source/USDClasses/Private/USDMaterialUtils.cpp
--- a/Engine/Plugins/Runtime/USDCore/Source/USDClasses/Private/USDMaterialUtils.cpp
+++ b/Engine/Plugins/Runtime/USDCore/Source/USDClasses/Private/USDMaterialUtils.cpp
@@ -297,3 +297,29 @@ bool UsdUnreal::MaterialUtils::IsReferencePreviewSurfaceMaterial(const FSoftObje
 
 	return ReferenceMaterials.Contains(Material);
 }
+
+namespace UE::USDMaterialUtils::Private
+{
+	static TArray<FName> RegisteredRenderContexts;
+}
+
+void UsdUnreal::MaterialUtils::RegisterRenderContext(const FName& RenderContextName)
+{
+	UE::USDMaterialUtils::Private::RegisteredRenderContexts.AddUnique(RenderContextName);
+	UE::USDMaterialUtils::Private::RegisteredRenderContexts.Sort(
+		[](const FName& LHS, const FName& RHS)
+		{
+			return LHS.ToString() < RHS.ToString();
+		}
+	);
+}
+
+void UsdUnreal::MaterialUtils::UnregisterRenderContext(const FName& RenderContextName)
+{
+	UE::USDMaterialUtils::Private::RegisteredRenderContexts.Remove(RenderContextName);
+}
+
+const TArray<FName>& UsdUnreal::MaterialUtils::GetRegisteredRenderContexts()
+{
+	return UE::USDMaterialUtils::Private::RegisteredRenderContexts;
+}
diff --git a/Engine/Plugins/Runtime/USDCore/Source/USDClasses/Public/USDMaterialUtils.h b/Engine/Plugins/Runtime/USDCore/Source/USDClasses/Public/USDMaterialUtils.h
--- a/Engine/Plugins/Runtime/USDCore/Source/USDClasses/Public/USDMaterialUtils.h
+++ b/Engine/Plugins/Runtime/USDCore/Source/USDClasses/Public/USDMaterialUtils.h
@@ -57,4 +57,8 @@ namespace UsdUnreal::MaterialUtils
 	// Returns whether the Material is one of the UsdPreviewSurface reference materials (which can be reassigned by the
 	// user on a per project basis)
 	USDCLASSES_API bool IsReferencePreviewSurfaceMaterial(const FSoftObjectPath& ReferenceMaterial);
+
+	USDCLASSES_API void RegisterRenderContext(const FName& RenderContextName);
+	USDCLASSES_API void UnregisterRenderContext(const FName& RenderContextName);
+	USDCLASSES_API const TArray<FName>& GetRegisteredRenderContexts();
 };	  // namespace UsdUnreal::MaterialUtils
diff --git a/Engine/Plugins/Runtime/USDCore/Source/UnrealUSDWrapper/Private/UnrealUSDWrapper.cpp b/Engine/Plugins/Runtime/USDCore/Source/UnrealUSDWrapper/Private/UnrealUSDWrapper.cpp
--- a/Engine/Plugins/Runtime/USDCore/Source/UnrealUSDWrapper/Private/UnrealUSDWrapper.cpp
+++ b/Engine/Plugins/Runtime/USDCore/Source/UnrealUSDWrapper/Private/UnrealUSDWrapper.cpp
@@ -4,6 +4,7 @@
 
 #include "USDClassesModule.h"
 #include "USDLog.h"
+#include "USDMaterialUtils.h"
 #include "USDMemory.h"
 #include "USDProjectSettings.h"
 #include "USDReferenceOptions.h"
@@ -475,9 +476,18 @@ FString UnrealIdentifiers::ModelApplyDrawMode = TEXT("model:applyDrawMode");
 FString UnrealIdentifiers::UsdNamespaceDelimiter = TEXT(":");
 #endif	  // USE_USD_SDK
 
-#if WITH_EDITOR
-const FName UnrealIdentifiers::MaterialXRenderContext = TEXT("mtlx");
+#if USE_USD_SDK
+// This is largely for show, as the universalRenderContext is an empty string, and FNames created from empty strings equal NAME_None anyway
+// i.e.: NAME_None == FName(TEXT("")) == FName(TEXT("None")) == UnrealIdentifiers::UniversalRenderContext
+const FName UnrealIdentifiers::UniversalRenderContext = ANSI_TO_TCHAR(pxr::UsdShadeTokens->universalRenderContext.GetString().c_str());
+#else
+const FName UnrealIdentifiers::UniversalRenderContext = NAME_None;
 #endif
+const FName UnrealIdentifiers::UnrealRenderContext = TEXT("unreal");
+const FName UnrealIdentifiers::MaterialXRenderContext = TEXT("mtlx");
+const FName UnrealIdentifiers::MdlRenderContext = TEXT("mdl");
+
+const FString UnrealIdentifiers::UniversalRenderContextDisplayString = TEXT("universal");
 
 FUsdDelegates::FUsdImportDelegate FUsdDelegates::OnPreUsdImport;
 FUsdDelegates::FUsdImportDelegate FUsdDelegates::OnPostUsdImport;
@@ -1239,6 +1249,9 @@ public:
 
 		FUsdMemoryManager::Initialize();
 
+		UsdUnreal::MaterialUtils::RegisterRenderContext(UnrealIdentifiers::UniversalRenderContext);
+		UsdUnreal::MaterialUtils::RegisterRenderContext(UnrealIdentifiers::UnrealRenderContext);
+
 #if WITH_EDITOR
 		// Update the supported filetype filters for reference/payload picker dialogs
 		for (TFieldIterator<FProperty> PropertyIterator(UUsdReferenceOptions::StaticClass()); PropertyIterator; ++PropertyIterator)
@@ -1266,6 +1279,9 @@ public:
 
 	virtual void ShutdownModule() override
 	{
+		UsdUnreal::MaterialUtils::UnregisterRenderContext(UnrealIdentifiers::UniversalRenderContext);
+		UsdUnreal::MaterialUtils::UnregisterRenderContext(UnrealIdentifiers::UnrealRenderContext);
+
 #if WITH_EDITOR
 		// We can't query default objects during engine exit
 		if (UObjectInitialized())
diff --git a/Engine/Plugins/Runtime/USDCore/Source/UnrealUSDWrapper/Public/UnrealUSDWrapper.h b/Engine/Plugins/Runtime/USDCore/Source/UnrealUSDWrapper/Public/UnrealUSDWrapper.h
--- a/Engine/Plugins/Runtime/USDCore/Source/UnrealUSDWrapper/Public/UnrealUSDWrapper.h
+++ b/Engine/Plugins/Runtime/USDCore/Source/UnrealUSDWrapper/Public/UnrealUSDWrapper.h
@@ -555,9 +555,12 @@ namespace UnrealIdentifiers
 	// The character used to separate property namespaces (usually just ':')
 	extern UNREALUSDWRAPPER_API FString UsdNamespaceDelimiter;
 
-#if WITH_EDITOR
+	extern UNREALUSDWRAPPER_API const FName UniversalRenderContext; // Likely the empty string
+	extern UNREALUSDWRAPPER_API const FName UnrealRenderContext;
 	extern UNREALUSDWRAPPER_API const FName MaterialXRenderContext;
-#endif
+	extern UNREALUSDWRAPPER_API const FName MdlRenderContext;
+
+	extern UNREALUSDWRAPPER_API const FString UniversalRenderContextDisplayString; // "universal"
 
 }	 // namespace UnrealIdentifiers
 
