Add ability for user to reset an imported level to values of the last import
- Assets imported through Interchange now has an entry in context menu to reset properties.
- Reset can be triggered from the outliner or content browser.
- Added ability for the factory nodes to track attribute delegates and be refilled later.

Other Fixes:
- Refactored the ActorFactory to handle actors in a more generic way.

diff --git a/Engine/Plugins/Interchange/Editor/Source/InterchangeEditor/InterchangeEditor.Build.cs b/Engine/Plugins/Interchange/Editor/Source/InterchangeEditor/InterchangeEditor.Build.cs
--- a/Engine/Plugins/Interchange/Editor/Source/InterchangeEditor/InterchangeEditor.Build.cs
+++ b/Engine/Plugins/Interchange/Editor/Source/InterchangeEditor/InterchangeEditor.Build.cs
@@ -32,7 +32,9 @@ namespace UnrealBuildTool.Rules
 					"InterchangeImport",
 					"InterchangeNodes",
 					"InterchangePipelines",
+					"LevelEditor",
 					"MessageLog",
+					"Slate",
 					"SlateCore",
 					"InputCore", // Translator settings customizations
 					"Slate", // Translator settings customizations
diff --git a/Engine/Plugins/Interchange/Editor/Source/InterchangeEditor/Private/AssetDefinition_InterchangeSceneImportAsset.cpp b/Engine/Plugins/Interchange/Editor/Source/InterchangeEditor/Private/AssetDefinition_InterchangeSceneImportAsset.cpp
--- a/Engine/Plugins/Interchange/Editor/Source/InterchangeEditor/Private/AssetDefinition_InterchangeSceneImportAsset.cpp
+++ b/Engine/Plugins/Interchange/Editor/Source/InterchangeEditor/Private/AssetDefinition_InterchangeSceneImportAsset.cpp
@@ -3,18 +3,30 @@
 #include "AssetDefinition_InterchangeSceneImportAsset.h"
 
 #include "InterchangeAssetImportData.h"
+#include "InterchangeAssetUserData.h"
 #include "InterchangeEditorModule.h"
+#include "InterchangeImportReset.h"
+#include "InterchangeLevelFactoryNode.h"
 #include "InterchangeManager.h"
 #include "InterchangeSceneImportAsset.h"
 
 #include "ContentBrowserMenuContexts.h"
+#include "EditorReimportHandler.h"
+#include "Engine/World.h"
+#include "Engine/Level.h"
+#include "GameFramework/WorldSettings.h"
 #include "LevelInstance/LevelInstanceActor.h"
+#include "LevelInstance/LevelInstanceComponent.h"
 #include "Misc/App.h"
 #include "Misc/DelayedAutoRegister.h"
 #include "ThumbnailRendering/SceneThumbnailInfo.h"
 #include "ToolMenus.h"
 #include "ToolMenuSection.h"
-#include "EditorReimportHandler.h"
+
+#if WITH_EDITOR
+#include "Editor/EditorEngine.h"
+extern UNREALED_API UEditorEngine* GEditor;
+#endif
 
 #define LOCTEXT_NAMESPACE "AssetDefinition_InterchangeSceneImportAsset"
 
@@ -32,6 +44,58 @@ UThumbnailInfo* UAssetDefinition_InterchangeSceneImportAsset::LoadThumbnailInfo(
 
 namespace MenuExtension_InterchangeSceneImportAsset
 {
+	ALevelInstance* SpawnLevelInstanceInEditor(FName AssetName)
+	{
+		FWorldContext& EditorWorldContext = GEditor->GetEditorWorldContext();
+		if (UWorld* EditorWorld = EditorWorldContext.World())
+		{
+			FActorSpawnParameters SpawnParameters;
+			SpawnParameters.Name = AssetName;
+			SpawnParameters.NameMode = FActorSpawnParameters::ESpawnActorNameMode::Requested;
+			SpawnParameters.OverrideLevel = EditorWorld->PersistentLevel;
+			SpawnParameters.SpawnCollisionHandlingOverride = ESpawnActorCollisionHandlingMethod::AlwaysSpawn;
+
+			return EditorWorld->SpawnActor<ALevelInstance>(ALevelInstance::StaticClass(), SpawnParameters);
+		}
+
+		return nullptr;
+	}
+
+	void InitializeLevelInstance(ALevelInstance* LevelInstanceActor, UWorld* ReferencedWorld)
+	{
+		if (!LevelInstanceActor->GetRootComponent())
+		{
+			USceneComponent* RootComponent = NewObject<USceneComponent>(LevelInstanceActor, USceneComponent::GetDefaultSceneRootVariableName(), RF_Transactional);
+#if WITH_EDITORONLY_DATA
+			RootComponent->bVisualizeComponent = true;
+#endif
+			LevelInstanceActor->SetRootComponent(RootComponent);
+			LevelInstanceActor->AddInstanceComponent(RootComponent);
+		}
+
+#if WITH_EDITORONLY_DATA
+		UWorld* ParentWorld = LevelInstanceActor->GetWorld();
+
+		// Make sure newly created level asset gets scanned
+		ULevel::ScanLevelAssets(ReferencedWorld->GetPackage()->GetName());
+
+		ParentWorld->PreEditChange(nullptr);
+
+		LevelInstanceActor->SetWorldAsset(ReferencedWorld);
+		LevelInstanceActor->UpdateLevelInstanceFromWorldAsset();
+		LevelInstanceActor->LoadLevelInstance();
+
+		//Reference world must be cleanup since they are not the main world.
+		//This remove all the world managers and prevent GC issue when unloading the main world referencing this world.
+		if (ReferencedWorld->bIsWorldInitialized)
+		{
+			ReferencedWorld->CleanupWorld();
+		}
+
+		ParentWorld->PostEditChange();
+#endif
+	}
+
 	ULevel* GetOriginalLevel(UInterchangeSceneImportAsset* Asset)
 	{
 		TArray<FSoftObjectPath> SceneSoftObjectPaths;
@@ -49,6 +113,126 @@ namespace MenuExtension_InterchangeSceneImportAsset
 		return nullptr;
 	}
 
+	void ExecuteResetOneAsset(UInterchangeSceneImportAsset* Asset, const FString& FilePath)
+	{
+		TRACE_CPUPROFILER_EVENT_SCOPE(Interchange_SceneImportAsset_ContextMenuReset)
+		FInterchangeResetParameters ResetParameters(Asset);
+
+		const FSoftObjectPath CurrentSceneImportObjectPath(Asset);
+		const TSubclassOf<UAssetUserData> AssetUserDataClass(UInterchangeLevelAssetUserData::StaticClass());
+
+		ResetParameters.PreResetDelegates.OnNodeProcessed.BindLambda([SceneImportAsset = Asset, &CurrentSceneImportObjectPath, &AssetUserDataClass](FInterchangeResetParameters& ResetParameters, const UInterchangeFactoryBase* Factory, const UInterchangeFactoryBaseNode* FactoryNode)
+		{
+			if (const UInterchangeLevelFactoryNode* LevelFactoryNode = Cast<UInterchangeLevelFactoryNode>(FactoryNode))
+			{
+				bool bShouldCreateLevel = false;
+				LevelFactoryNode->GetCustomShouldCreateLevel(bShouldCreateLevel);
+								
+				if(bShouldCreateLevel)
+				{
+					FSoftObjectPath WorldPath;
+					if (LevelFactoryNode->GetCustomReferenceObject(WorldPath))
+					{
+						if (UWorld* ReferencedWorld = Cast<UWorld>(WorldPath.TryLoad()))
+						{
+							if (AWorldSettings* WorldSettings = ReferencedWorld->GetWorldSettings())
+							{
+								if (UInterchangeLevelAssetUserData* AssetUserData = Cast<UInterchangeLevelAssetUserData>(WorldSettings->GetAssetUserDataOfClass(AssetUserDataClass)))
+								{
+									bool bShouldCreateLevelInstance = false;
+									for (const FSoftObjectPath& SceneImportAssetPath : AssetUserData->SceneImportPaths)
+									{
+										if (SceneImportAssetPath == CurrentSceneImportObjectPath)
+										{
+											bShouldCreateLevelInstance = true;
+											break;
+										}
+									}
+
+									if (!bShouldCreateLevelInstance)
+									{
+										return;
+									}
+
+									TObjectPtr<ALevelInstance> LevelInstance = SpawnLevelInstanceInEditor(FName(ReferencedWorld->GetName()));
+									if (IsValid(LevelInstance))
+									{
+										InitializeLevelInstance(LevelInstance, ReferencedWorld);
+										if (!ResetParameters.ResetContextData->ObjectsSpawnedDuringReset.Contains(FactoryNode))
+										{
+											ResetParameters.ResetContextData->ObjectsSpawnedDuringReset.Emplace(FactoryNode, TArray<TObjectPtr<UObject>>());
+										}
+										
+										ResetParameters.ResetContextData->ObjectsSpawnedDuringReset[FactoryNode].Add(LevelInstance);
+
+										if (LevelInstance->CanEnterEdit())
+										{
+											LevelInstance->EnterEdit();
+										}
+									}
+								}
+							}
+						}
+					}
+				}
+			}
+		});
+
+		ResetParameters.PostResetDelegates.OnNodeProcessed.BindLambda([SceneImportAsset = Asset, &CurrentSceneImportObjectPath, &AssetUserDataClass](FInterchangeResetParameters& ResetParameters, const UInterchangeFactoryBase* Factory, const UInterchangeFactoryBaseNode* FactoryNode)
+			{
+				if (const UInterchangeLevelFactoryNode* LevelFactoryNode = Cast<UInterchangeLevelFactoryNode>(FactoryNode))
+				{
+					bool bShouldCreateLevel = false;
+					LevelFactoryNode->GetCustomShouldCreateLevel(bShouldCreateLevel);
+
+					if (bShouldCreateLevel)
+					{
+						if (ResetParameters.ResetContextData->ObjectsSpawnedDuringReset.Contains(FactoryNode))
+						{
+							for (TObjectPtr<UObject>& SpawnedObject : ResetParameters.ResetContextData->ObjectsSpawnedDuringReset[FactoryNode])
+							{
+								if (ALevelInstance* LevelInstance = Cast<ALevelInstance>(SpawnedObject))
+								{
+									if (LevelInstance->IsEditing())
+									{
+#if WITH_EDITOR
+										// Make sure transformation is up to date after registration as its possible LevelInstance actor can get unregistered when editing properties
+										// through Details panel. In this case the ULevelInstanceComponent might not be able to update the ALevelInstanceEditorInstanceActor transform.
+										LevelInstance->GetLevelInstanceComponent()->UpdateEditorInstanceActor();
+#endif
+										if (LevelInstance->CanExitEdit())
+										{
+											LevelInstance->ExitEdit();
+											// Clear selection in case anything gets selected after applying the changes.
+											GEditor->SelectNone(true, true);
+										}
+									}
+									LevelInstance->Destroy();
+								}
+							}
+						}
+					}
+				}
+			});
+
+		FInterchangeReset::ExecuteReset(ResetParameters);
+		GEditor->RedrawAllViewports();
+	}
+
+	void ExecuteResetScene(const FToolMenuContext& InContext)
+	{
+		const UContentBrowserAssetContextMenuContext* Context = UContentBrowserAssetContextMenuContext::FindContextWithAssets(InContext);
+		TArray<UInterchangeSceneImportAsset*> SceneImportAssets = Context->LoadSelectedObjects<UInterchangeSceneImportAsset>();
+
+		for (UInterchangeSceneImportAsset* SceneImportAsset : SceneImportAssets)
+		{
+			if (SceneImportAsset)
+			{
+				ExecuteResetOneAsset(SceneImportAsset, SceneImportAsset->AssetImportData->GetFirstFilename());
+			}
+		}
+	}
+
 	void ExecuteReimportOneAsset(UInterchangeSceneImportAsset* Asset, const FString& FilePath)
 	{
 		using namespace UE::Interchange;
@@ -94,7 +278,7 @@ namespace MenuExtension_InterchangeSceneImportAsset
 		{
 			return;
 		}
-		
+
 #if WITH_EDITORONLY_DATA
 		UInterchangeSceneImportAsset* SceneImportAsset = SceneImportAssets[0];
 		if (!SceneImportAsset->AssetImportData)
@@ -111,6 +295,25 @@ namespace MenuExtension_InterchangeSceneImportAsset
 #endif
 	}
 
+	bool CanExecuteResetScene(const FToolMenuContext& InContext)
+	{
+		const UContentBrowserAssetContextMenuContext* Context = UContentBrowserAssetContextMenuContext::FindContextWithAssets(InContext);
+		TArray<UInterchangeSceneImportAsset*> SceneImportAssets = Context->LoadSelectedObjects<UInterchangeSceneImportAsset>();
+
+		for (UInterchangeSceneImportAsset* SceneImportAsset : SceneImportAssets)
+		{
+			if (SceneImportAsset && SceneImportAsset->AssetImportData)
+			{
+				if (SceneImportAsset->AssetImportData->ExtractFilenames().Num() > 0)
+				{
+					return true;
+				}
+			}
+		}
+
+		return false;
+	}
+
 	bool CanExecuteReimport(const FToolMenuContext& InContext)
 	{
 		const UContentBrowserAssetContextMenuContext* Context = UContentBrowserAssetContextMenuContext::FindContextWithAssets(InContext);
@@ -148,45 +351,56 @@ namespace MenuExtension_InterchangeSceneImportAsset
 		return false;
 	}
 
-	static FDelayedAutoRegisterHelper DelayedAutoRegister(EDelayedRegisterRunPhase::EndOfEngineInit, []{ 
+	static FDelayedAutoRegisterHelper DelayedAutoRegister(EDelayedRegisterRunPhase::EndOfEngineInit, [] {
 		UToolMenus::RegisterStartupCallback(FSimpleMulticastDelegate::FDelegate::CreateLambda([]()
-		{
-			FToolMenuOwnerScoped OwnerScoped(UE_MODULE_NAME);
-			UToolMenu* Menu = UE::ContentBrowser::ExtendToolMenu_AssetContextMenu(UInterchangeSceneImportAsset::StaticClass());
-		
-			if (!Menu->FindSection("Interchange"))
-			{
-				const FToolMenuInsert MenuInsert(NAME_None, EToolMenuInsertType::First);
-				Menu->AddSection("Interchange", LOCTEXT("InterchangeSceneImportAsset_Section", "Interchange"), MenuInsert);
-			}
-			FToolMenuSection& Section = Menu->FindOrAddSection("Interchange");
-			Section.AddDynamicEntry(NAME_None, FNewToolMenuSectionDelegate::CreateLambda([](FToolMenuSection& InSection)
 			{
+				FToolMenuOwnerScoped OwnerScoped(UE_MODULE_NAME);
+				UToolMenu* Menu = UE::ContentBrowser::ExtendToolMenu_AssetContextMenu(UInterchangeSceneImportAsset::StaticClass());
+
+				if (!Menu->FindSection("Interchange"))
 				{
-					const TAttribute<FText> Label = LOCTEXT("InterchangeSceneImportAsset_Reimport", "Reimport Scene");
-					const TAttribute<FText> ToolTip = LOCTEXT("InterchangeSceneImportAsset_ReimportTooltip", "Reimport the scene associated with each selected InterchangeSceneImportAsset.");
-					const FSlateIcon Icon = FSlateIcon(FAppStyle::GetAppStyleSetName(), "Icons.Import");
+					const FToolMenuInsert MenuInsert(NAME_None, EToolMenuInsertType::First);
+					Menu->AddSection("Interchange", LOCTEXT("InterchangeSceneImportAsset_Section", "Interchange"), MenuInsert);
+				}
+				FToolMenuSection& Section = Menu->FindOrAddSection("Interchange");
+				Section.AddDynamicEntry(NAME_None, FNewToolMenuSectionDelegate::CreateLambda([](FToolMenuSection& InSection)
+					{
+						{
+							const TAttribute<FText> Label = LOCTEXT("InterchangeSceneImportAsset_ResetScene", "Reset Scene");
+							const TAttribute<FText> ToolTip = LOCTEXT("InterchangeSceneImportAsset_ResetSceneTooltip", "Reset the scene associated with each selected InterchangeSceneImportAsset to original pipeline properties.");
+							const FSlateIcon Icon = FSlateIcon(FAppStyle::GetAppStyleSetName(), "SourceControl.Actions.Refresh");
 
-					FToolUIAction UIAction;
-					UIAction.ExecuteAction = FToolMenuExecuteAction::CreateStatic(&ExecuteReimport);
-					UIAction.CanExecuteAction = FToolMenuCanExecuteAction::CreateStatic(&CanExecuteReimport);
+							FToolUIAction UIAction;
+							UIAction.ExecuteAction = FToolMenuExecuteAction::CreateStatic(&ExecuteResetScene);
+							UIAction.CanExecuteAction = FToolMenuCanExecuteAction::CreateStatic(&CanExecuteResetScene);
 
-					InSection.AddMenuEntry("InterchangeSceneImportAsset_Reimport", Label, ToolTip, Icon, UIAction);
-				}
-				{
-					const TAttribute<FText> Label = LOCTEXT("InterchangeSceneImportAsset_ReimportWithFile", "Reimport Scene With File");
-					const TAttribute<FText> ToolTip = LOCTEXT("InterchangeSceneImportAsset_ReimportWithFile_Tooltip", "Reimport the scene associated with the selected InterchangeSceneImportAsset using a new file.");
-					const FSlateIcon Icon = FSlateIcon(FAppStyle::GetAppStyleSetName(), "Icons.Import");
+							InSection.AddMenuEntry("InterchangeSceneImportAsset_ResetScene", Label, ToolTip, Icon, UIAction);
+						}
+						{
+							const TAttribute<FText> Label = LOCTEXT("InterchangeSceneImportAsset_Reimport", "Reimport Scene");
+							const TAttribute<FText> ToolTip = LOCTEXT("InterchangeSceneImportAsset_ReimportTooltip", "Reimport the scene associated with each selected InterchangeSceneImportAsset.");
+							const FSlateIcon Icon = FSlateIcon(FAppStyle::GetAppStyleSetName(), "Icons.Import");
 
-					FToolUIAction UIAction;
-					UIAction.ExecuteAction = FToolMenuExecuteAction::CreateStatic(&ExecuteReimportWithFile);
-					UIAction.CanExecuteAction = FToolMenuCanExecuteAction::CreateStatic(&CanExecuteReimportWithFile);
+							FToolUIAction UIAction;
+							UIAction.ExecuteAction = FToolMenuExecuteAction::CreateStatic(&ExecuteReimport);
+							UIAction.CanExecuteAction = FToolMenuCanExecuteAction::CreateStatic(&CanExecuteReimport);
 
-					InSection.AddMenuEntry("InterchangeSceneImportAsset_ReimportWithFile", Label, ToolTip, Icon, UIAction);
-				}
+							InSection.AddMenuEntry("InterchangeSceneImportAsset_Reimport", Label, ToolTip, Icon, UIAction);
+						}
+						{
+							const TAttribute<FText> Label = LOCTEXT("InterchangeSceneImportAsset_ReimportWithFile", "Reimport Scene With File");
+							const TAttribute<FText> ToolTip = LOCTEXT("InterchangeSceneImportAsset_ReimportWithFile_Tooltip", "Reimport the scene associated with the selected InterchangeSceneImportAsset using a new file.");
+							const FSlateIcon Icon = FSlateIcon(FAppStyle::GetAppStyleSetName(), "Icons.Import");
+
+							FToolUIAction UIAction;
+							UIAction.ExecuteAction = FToolMenuExecuteAction::CreateStatic(&ExecuteReimportWithFile);
+							UIAction.CanExecuteAction = FToolMenuCanExecuteAction::CreateStatic(&CanExecuteReimportWithFile);
+
+							InSection.AddMenuEntry("InterchangeSceneImportAsset_ReimportWithFile", Label, ToolTip, Icon, UIAction);
+						}
+					}));
 			}));
-		}));
-	});
+		});
 }
 
 #undef LOCTEXT_NAMESPACE
diff --git a/Engine/Plugins/Interchange/Editor/Source/InterchangeEditor/Private/InterchangeEditorModule.cpp b/Engine/Plugins/Interchange/Editor/Source/InterchangeEditor/Private/InterchangeEditorModule.cpp
--- a/Engine/Plugins/Interchange/Editor/Source/InterchangeEditor/Private/InterchangeEditorModule.cpp
+++ b/Engine/Plugins/Interchange/Editor/Source/InterchangeEditor/Private/InterchangeEditorModule.cpp
@@ -8,6 +8,7 @@
 #include "InterchangeFbxAssetImportDataConverter.h"
 #include "InterchangeManager.h"
 #include "InterchangeUsdTranslatorSettingsCustomization.h"
+#include "InterchangeResetContextMenuExtender.h"
 #include "MessageLogModule.h"
 #include "Misc/App.h"
 #include "Misc/CoreDelegates.h"
@@ -150,6 +151,8 @@ void FInterchangeEditorModule::StartupModule()
 	{
 		FCoreDelegates::OnPostEngineInit.AddLambda(RegisterItems);
 	}
+
+	FInterchangeResetContextMenuExtender::SetupLevelEditorContextMenuExtender();
 }
 
 void FInterchangeEditorModule::ShutdownModule()
@@ -159,6 +162,7 @@ void FInterchangeEditorModule::ShutdownModule()
 	{
 		PropertyModule->UnregisterCustomClassLayout(TEXT("InterchangeUsdTranslatorSettings"));
 	}
+	FInterchangeResetContextMenuExtender::RemoveLevelEditorContextMenuExtender();
 }
 
 IMPLEMENT_MODULE(FInterchangeEditorModule, InterchangeEditor)
diff --git a/Engine/Plugins/Interchange/Editor/Source/InterchangeEditor/Private/InterchangeResetContextMenuExtender.cpp b/Engine/Plugins/Interchange/Editor/Source/InterchangeEditor/Private/InterchangeResetContextMenuExtender.cpp
new file mode 100644
--- /dev/null
+++ b/Engine/Plugins/Interchange/Editor/Source/InterchangeEditor/Private/InterchangeResetContextMenuExtender.cpp
@@ -0,0 +1,464 @@
+// Copyright Epic Games, Inc. All Rights Reserved.
+
+#include "InterchangeResetContextMenuExtender.h"
+
+#include "Algo/AnyOf.h"
+#include "AssetToolsModule.h"
+#include "ContentBrowserMenuContexts.h"
+#include "Editor.h"
+#include "Editor/EditorEngine.h"
+#include "Engine/Engine.h"
+#include "Engine/Selection.h"
+#include "EngineUtils.h"
+#include "Framework/MultiBox/MultiBoxBuilder.h"
+#include "InterchangeAssetUserData.h"
+#include "InterchangeImportReset.h"
+#include "InterchangeLevelFactoryNode.h"
+#include "InterchangeManager.h"
+#include "InterchangeSceneImportAsset.h"
+#include "LevelEditor.h"
+#include "LevelEditorMenuContext.h"
+#include "LevelInstance/LevelInstanceActor.h"
+#include "LevelInstance/LevelInstanceComponent.h"
+#include "LevelInstance/LevelInstanceInterface.h"
+#include "LevelInstance/LevelInstanceEditorPivotActor.h"
+#include "LevelInstance/LevelInstanceSettings.h"
+#include "LevelInstance/LevelInstanceSubsystem.h"
+#include "ToolMenu.h"
+#include "ToolMenuDelegates.h"
+#include "ToolMenus.h"
+#include "ToolMenuSection.h"
+
+#define LOCTEXT_NAMESPACE "InterchangeResetContextMenuHandler"
+
+namespace UE::Interchange::InterchangeReset
+{
+	bool CanResetWorld(const UWorld* World)
+	{
+		if (AWorldSettings* WorldSettings = World->GetWorldSettings())
+		{
+			return WorldSettings->HasAssetUserDataOfClass(TSubclassOf<UAssetUserData>(UInterchangeLevelAssetUserData::StaticClass()));
+		}
+
+		return false;
+	}
+
+	bool CanResetActor(const AActor* Actor)
+	{
+		if (const ALevelInstance* LevelInstanceActor = Cast<ALevelInstance>(Actor))
+		{
+			const TSoftObjectPtr<UWorld>& WorldObjectPtr = LevelInstanceActor->GetWorldAsset();
+			if (const UWorld* ReferencedWorld = Cast<UWorld>(WorldObjectPtr.ToSoftObjectPath().TryLoad()))
+			{
+				return CanResetWorld(ReferencedWorld);
+			}
+		}
+		else
+		{
+			TSubclassOf<UAssetUserData> InterchangeAssetUserDataClass(UInterchangeAssetUserData::StaticClass());
+			for (UActorComponent* Component : Actor->GetComponents())
+			{
+				if (Component->HasAssetUserDataOfClass(InterchangeAssetUserDataClass))
+				{
+					return true;
+				}
+			}
+		}
+
+		return false;
+	}
+
+	UInterchangeSceneImportAsset* GetSceneImportAssetFromPath(const FSoftObjectPath& SoftObjectPath)
+	{
+		if (UInterchangeSceneImportAsset* SceneImportAsset = Cast<UInterchangeSceneImportAsset>(SoftObjectPath.TryLoad()))
+		{
+			return SceneImportAsset;
+		}
+
+		return nullptr;
+	}
+
+
+	UInterchangeSceneImportAsset* GetSceneImportAssetForWorld(const UWorld* ReferencedWorld)
+	{
+		if (AWorldSettings* WorldSettings = ReferencedWorld->GetWorldSettings())
+		{
+			if (UInterchangeLevelAssetUserData* AssetUserData = Cast<UInterchangeLevelAssetUserData>(WorldSettings->GetAssetUserDataOfClass(TSubclassOf<UAssetUserData>(UInterchangeLevelAssetUserData::StaticClass()))))
+			{
+				if (!AssetUserData->SceneImportPaths.IsEmpty())
+				{
+					return GetSceneImportAssetFromPath(AssetUserData->SceneImportPaths[0]);
+				}
+			}
+		}
+
+		return nullptr;
+	}
+
+	UInterchangeSceneImportAsset* GetSceneImportAsset(AActor* ImportedActor)
+	{
+		if (const ALevelInstance* LevelInstanceActor = Cast<ALevelInstance>(ImportedActor))
+		{
+			const TSoftObjectPtr<UWorld>& WorldObjectPtr = LevelInstanceActor->GetWorldAsset();
+			if (const UWorld* ReferencedWorld = Cast<UWorld>(WorldObjectPtr.ToSoftObjectPath().TryLoad()))
+			{
+				if (UInterchangeSceneImportAsset* SceneImportAsset = GetSceneImportAssetForWorld(ReferencedWorld))
+				{
+					return SceneImportAsset;
+				}
+			}
+		}
+		else
+		{
+			const TSubclassOf<UAssetUserData> InterchangeAssetUserDataClass(UInterchangeAssetUserData::StaticClass());
+			
+			for (UActorComponent* Component : ImportedActor->GetComponents())
+			{
+				if (UInterchangeAssetUserData* AssetUserData = Cast<UInterchangeAssetUserData>(Component->GetAssetUserDataOfClass(InterchangeAssetUserDataClass)))
+				{
+					if (const FString* SceneImportAssetPathString = AssetUserData->MetaData.Find(Constants::SceneImportAssetPathKey))
+					{
+						if (UInterchangeSceneImportAsset* SceneImportAsset = GetSceneImportAssetFromPath(*SceneImportAssetPathString);
+							SceneImportAsset && SceneImportAsset->ContainsObject(ImportedActor))
+						{
+							return SceneImportAsset;
+						}
+					}
+				}
+			}
+		}
+
+		return nullptr;
+	}
+	
+	bool ExecuteLevelInstanceReset(const UInterchangeSceneImportAsset* SceneImportAsset, ALevelInstance* LevelInstanceActor)
+	{
+		if (!LevelInstanceActor || !SceneImportAsset)
+		{
+			return false;
+		}
+
+		ILevelInstanceInterface* LevelInstance = Cast<ILevelInstanceInterface>(LevelInstanceActor);
+
+		const TSoftObjectPtr<UWorld>& WorldObjectPtr = LevelInstanceActor->GetWorldAsset();
+
+		if (const UWorld* World = Cast<UWorld>(WorldObjectPtr.ToSoftObjectPath().TryLoad()))
+		{
+			const bool bWasPrevEditing = LevelInstance->IsEditing();
+			if (!bWasPrevEditing && LevelInstance->CanEnterEdit())
+			{
+				LevelInstance->EnterEdit(LevelInstanceActor);
+			}
+
+			FInterchangeResetParameters ResetParameters(SceneImportAsset);
+
+			// Filter nodes
+			if (CCvarInterchangeResetFilteredNodes->GetBool())
+			{
+				if (ULevel* LoadedLevel = LevelInstance->GetLoadedLevel())
+				{
+					TArray<AActor*> SceneActors = LoadedLevel->Actors;
+					SceneActors.Add(LevelInstanceActor);
+
+					TSet<const UInterchangeFactoryBaseNode*> TempFilteredFactoryNodes;
+					for (AActor* Child : SceneActors)
+					{
+						if (const UInterchangeFactoryBaseNode* FactoryNode = SceneImportAsset->GetFactoryNode(Child))
+						{
+							TempFilteredFactoryNodes.Add(FactoryNode);
+						}
+					}
+					ResetParameters.FilteredNodes = MoveTemp(TempFilteredFactoryNodes);
+				}
+			}
+
+			FInterchangeReset::ExecuteReset(ResetParameters);
+#if WITH_EDITOR
+			// Make sure transformation is up to date after registration as its possible LevelInstance actor can get unregistered when editing properties
+			// through Details panel. In this case the ULevelInstanceComponent might not be able to update the ALevelInstanceEditorInstanceActor transform.
+			LevelInstance->GetLevelInstanceComponent()->UpdateEditorInstanceActor();
+#endif
+			GEditor->RedrawAllViewports();
+
+			if (!bWasPrevEditing && LevelInstance->CanExitEdit())
+			{
+				LevelInstance->ExitEdit();
+			}
+
+			return true;
+		}
+		
+		return false;
+	}
+
+	bool ExecuteLevelInstanceReset(TPair<UInterchangeSceneImportAsset*, ALevelInstance*> ResetDataPair)
+	{
+		return ExecuteLevelInstanceReset(ResetDataPair.Key, ResetDataPair.Value);
+	}
+
+	void ResetActors(TArray<AActor*> SelectedActors)
+	{
+		TRACE_CPUPROFILER_EVENT_SCOPE(Interchange_Actors_ContextMenuReset)
+		TMap<UInterchangeSceneImportAsset*, FInterchangeResetParameters> BatchedReset;
+		TMap<UInterchangeSceneImportAsset*, ALevelInstance*> LevelInstanceResets;
+
+		for (AActor* Actor : SelectedActors)
+		{
+			if (!Actor)
+			{
+				continue;
+			}
+
+			if (UInterchangeSceneImportAsset* SceneImportAsset = GetSceneImportAsset(Actor))
+			{
+				if (ALevelInstance* LevelInstanceActor = Cast<ALevelInstance>(Actor))
+				{
+					if (!LevelInstanceResets.Contains(SceneImportAsset))
+					{
+						// It is enough to reset just one of the level instances as the changes would also be applied to all other when the reset is committed.
+						LevelInstanceResets.Emplace(SceneImportAsset, LevelInstanceActor);
+					}
+				}
+				else
+				{
+					if (!BatchedReset.Contains(SceneImportAsset))
+					{
+						BatchedReset.Emplace(SceneImportAsset, SceneImportAsset);
+					}
+
+					if (CCvarInterchangeResetFilteredNodes->GetBool())
+					{
+						if (const UInterchangeFactoryBaseNode* FactoryNode = SceneImportAsset->GetFactoryNode(Actor))
+						{
+							BatchedReset[SceneImportAsset].FilteredNodes.Add(FactoryNode);
+						}
+					}
+				}
+			}
+		}
+
+		for (TPair<UInterchangeSceneImportAsset*, FInterchangeResetParameters>& ResetDataPair : BatchedReset)
+		{
+			// As these are all actors in the scene, no need to do anything else to reset them.
+			FInterchangeReset::ExecuteReset(ResetDataPair.Value);
+		}
+
+		for (TPair<UInterchangeSceneImportAsset*, ALevelInstance*>& ResetDataPair : LevelInstanceResets)
+		{
+			ExecuteLevelInstanceReset(ResetDataPair);
+		}
+
+		GEditor->RedrawAllViewports();
+	}
+
+	TSharedRef<FExtender> OnExtendLevelEditorActorSelectionMenu(const TSharedRef<FUICommandList> CommandList, TArray<AActor*> SelectedActors)
+	{
+		TSharedRef<FExtender> Extender = MakeShared<FExtender>();
+
+		bool bShouldExtendActorActions = false;
+
+		for (AActor* Actor : SelectedActors)
+		{
+			if (CanResetActor(Actor))
+			{
+				bShouldExtendActorActions = true;
+				break;
+			}
+		}
+
+		if (bShouldExtendActorActions)
+		{
+			// Add the Interchange actions sub-menu extender
+			Extender->AddMenuExtension("ActorTypeTools", EExtensionHook::After, nullptr, FMenuExtensionDelegate::CreateLambda(
+				[SelectedActors](FMenuBuilder& MenuBuilder)
+				{
+					MenuBuilder.BeginSection("Interchange", LOCTEXT("InterchangeMenuSection", "Interchange"));
+					MenuBuilder.AddMenuEntry(
+						NSLOCTEXT("InterchangeActions", "ObjectContext_ResetInterchange", "Reset Properties"),
+						NSLOCTEXT("InterchangeActions", "ObjectContext_ResetInterchangeTooltip", "Resets overridden values with the values from Interchange Import"),
+						FSlateIcon(FAppStyle::GetAppStyleSetName(), "SourceControl.Actions.Refresh"),
+						FUIAction(FExecuteAction::CreateStatic(&ResetActors, SelectedActors), FCanExecuteAction()));
+					MenuBuilder.EndSection();
+				}));
+		}
+
+		return Extender;
+	}
+
+	namespace LevelContextMenuHelpers
+	{
+		void UpdateLevelInstanceActorWorld(ALevelInstance* LevelInstanceActor, UWorld* ReferencedWorld)
+		{
+#if WITH_EDITORONLY_DATA
+			UWorld* ParentWorld = LevelInstanceActor->GetWorld();
+
+			// Make sure newly created level asset gets scanned
+			ULevel::ScanLevelAssets(ReferencedWorld->GetPackage()->GetName());
+
+			ParentWorld->PreEditChange(nullptr);
+
+			LevelInstanceActor->SetWorldAsset(ReferencedWorld);
+			LevelInstanceActor->UpdateLevelInstanceFromWorldAsset();
+			LevelInstanceActor->LoadLevelInstance();
+
+			//Reference world must be cleanup since they are not the main world.
+			//This remove all the world managers and prevent GC issue when unloading the main world referencing this world.
+			if (ReferencedWorld->bIsWorldInitialized)
+			{
+				ReferencedWorld->CleanupWorld();
+			}
+
+			ParentWorld->PostEditChange();
+#endif
+		}
+
+		bool CanExecuteResetLevel(const FToolMenuContext& InContext)
+		{
+			const UContentBrowserAssetContextMenuContext* Context = UContentBrowserAssetContextMenuContext::FindContextWithAssets(InContext);
+			TArray<UWorld*> Worlds = Context->LoadSelectedObjects<UWorld>();
+			return Algo::AnyOf(Worlds, [](UWorld* World)
+				{
+					return World && CanResetWorld(World);
+				});
+		}
+
+		void ExecuteResetOneLevel(UWorld* World, const UInterchangeSceneImportAsset* SceneImportAsset)
+		{
+			TRACE_CPUPROFILER_EVENT_SCOPE(Interchange_Level_ContextMenuReset)
+			if (!World || !SceneImportAsset)
+			{
+				return;
+			}
+
+			FWorldContext& EditorWorldContext = GEditor->GetEditorWorldContext();
+			if (UWorld* EditorWorld = EditorWorldContext.World())
+			{
+				FActorSpawnParameters SpawnParameters;
+				SpawnParameters.Name = FName(World->GetMapName());
+				SpawnParameters.NameMode = FActorSpawnParameters::ESpawnActorNameMode::Requested;
+				SpawnParameters.OverrideLevel = EditorWorld->PersistentLevel;
+				SpawnParameters.SpawnCollisionHandlingOverride = ESpawnActorCollisionHandlingMethod::AlwaysSpawn;
+
+				if (ALevelInstance* SpawnedActor = EditorWorld->SpawnActor<ALevelInstance>(ALevelInstance::StaticClass(), SpawnParameters))
+				{
+					if (!SpawnedActor->GetRootComponent())
+					{
+						USceneComponent* RootComponent = NewObject<USceneComponent>(SpawnedActor, USceneComponent::GetDefaultSceneRootVariableName(), RF_Transactional);
+#if WITH_EDITORONLY_DATA
+						RootComponent->bVisualizeComponent = true;
+#endif
+						SpawnedActor->SetRootComponent(RootComponent);
+						SpawnedActor->AddInstanceComponent(RootComponent);
+					}
+
+					UpdateLevelInstanceActorWorld(SpawnedActor, World);
+					if (ExecuteLevelInstanceReset(SceneImportAsset, SpawnedActor))
+					{
+						GEditor->SelectNone(true, true);
+					}
+
+					SpawnedActor->Destroy();
+				}
+			}
+		}
+
+		void ExecuteResetLevel(const FToolMenuContext& InContext)
+		{
+			const UContentBrowserAssetContextMenuContext* Context = UContentBrowserAssetContextMenuContext::FindContextWithAssets(InContext);
+			TArray<UWorld*> Worlds = Context->LoadSelectedObjects<UWorld>();
+
+			for (UWorld* World : Worlds)
+			{
+				if (UInterchangeSceneImportAsset* SceneImportAsset = GetSceneImportAssetForWorld(World))
+				{
+					ExecuteResetOneLevel(World, SceneImportAsset);
+				}
+			}
+		}
+	}
+} //ns UE::Interchange::InterchangeReset
+
+static FDelayedAutoRegisterHelper DelayedAutoRegisterLevelContextMenu(EDelayedRegisterRunPhase::EndOfEngineInit, [] {
+	UToolMenus::RegisterStartupCallback(FSimpleMulticastDelegate::FDelegate::CreateLambda([]()
+		{
+			FToolMenuOwnerScoped OwnerScoped(UE_MODULE_NAME);
+			UToolMenu* Menu = UE::ContentBrowser::ExtendToolMenu_AssetContextMenu(UWorld::StaticClass());
+
+			if (!Menu->FindSection("Interchange"))
+			{
+				const FToolMenuInsert MenuInsert(NAME_None, EToolMenuInsertType::First);
+				Menu->AddSection("Interchange", LOCTEXT("Level_ResetSection", "Interchange"), MenuInsert);
+			}
+			FToolMenuSection& Section = Menu->FindOrAddSection("Interchange");
+
+			Section.AddDynamicEntry(NAME_None, FNewToolMenuSectionDelegate::CreateLambda([](FToolMenuSection& InSection)
+				{
+					const UContentBrowserAssetContextMenuContext* Context = UContentBrowserAssetContextMenuContext::FindContextWithAssets(InSection.Context);
+
+					if (Context->SelectedAssets.Num() != 1)
+					{
+						return;
+					}
+
+					bool bAddResetSection = false;
+
+					const FAssetData& WorldAsset = Context->SelectedAssets[0];
+					if (Context->SelectedAssets[0].IsInstanceOf<UWorld>())
+					{
+						if (UWorld* World = Cast<UWorld>(Context->SelectedAssets[0].GetAsset()))
+						{
+							if (AWorldSettings* WorldSettings = World->GetWorldSettings())
+							{
+								if (WorldSettings->HasAssetUserDataOfClass(TSubclassOf<UAssetUserData>(UInterchangeLevelAssetUserData::StaticClass())))
+								{
+									bAddResetSection = true;
+								}
+							}
+						}
+					}
+
+					if (bAddResetSection)
+					{
+						using namespace UE::Interchange::InterchangeReset;
+
+						const TAttribute<FText> Label = LOCTEXT("Level_ResetScene", "Reset Properties");
+						const TAttribute<FText> ToolTip = LOCTEXT("Level_ResetSceneTooltip", "Reset the level to original pipeline properties.");
+						const FSlateIcon Icon = FSlateIcon(FAppStyle::GetAppStyleSetName(), "SourceControl.Actions.Refresh");
+
+						FToolUIAction UIAction;
+						UIAction.ExecuteAction = FToolMenuExecuteAction::CreateStatic(&LevelContextMenuHelpers::ExecuteResetLevel);
+						UIAction.CanExecuteAction = FToolMenuCanExecuteAction::CreateStatic(&LevelContextMenuHelpers::CanExecuteResetLevel);
+
+						InSection.AddMenuEntry("Level_ResetScene", Label, ToolTip, Icon, UIAction);
+					}
+				}));
+		}));
+	});
+
+FDelegateHandle FInterchangeResetContextMenuExtender::LevelEditorExtenderDelegateHandle;
+
+void FInterchangeResetContextMenuExtender::SetupLevelEditorContextMenuExtender()
+{
+	FLevelEditorModule& LevelEditorModule = FModuleManager::LoadModuleChecked< FLevelEditorModule >("LevelEditor");
+	TArray< FLevelEditorModule::FLevelViewportMenuExtender_SelectedActors >& CBMenuExtenderDelegates = LevelEditorModule.GetAllLevelViewportContextMenuExtenders();
+
+	CBMenuExtenderDelegates.Add(FLevelEditorModule::FLevelViewportMenuExtender_SelectedActors::CreateStatic(&UE::Interchange::InterchangeReset::OnExtendLevelEditorActorSelectionMenu));
+	LevelEditorExtenderDelegateHandle = CBMenuExtenderDelegates.Last().GetHandle();
+}
+
+void FInterchangeResetContextMenuExtender::RemoveLevelEditorContextMenuExtender()
+{
+	if (LevelEditorExtenderDelegateHandle.IsValid() && FModuleManager::Get().IsModuleLoaded("LevelEditor"))
+	{
+		FLevelEditorModule& LevelEditorModule = FModuleManager::GetModuleChecked< FLevelEditorModule >("LevelEditor");
+		TArray< FLevelEditorModule::FLevelViewportMenuExtender_SelectedActors >& CBMenuExtenderDelegates = LevelEditorModule.GetAllLevelViewportContextMenuExtenders();
+		CBMenuExtenderDelegates.RemoveAll(
+			[&](const FLevelEditorModule::FLevelViewportMenuExtender_SelectedActors& Delegate)
+			{
+				return Delegate.GetHandle() == LevelEditorExtenderDelegateHandle;
+			}
+		);
+	}
+}
+
+#undef LOCTEXT_NAMESPACE // "InterchangeResetContextMenuHandler"
\ No newline at end of file
diff --git a/Engine/Plugins/Interchange/Editor/Source/InterchangeEditor/Private/InterchangeResetContextMenuExtender.h b/Engine/Plugins/Interchange/Editor/Source/InterchangeEditor/Private/InterchangeResetContextMenuExtender.h
new file mode 100644
--- /dev/null
+++ b/Engine/Plugins/Interchange/Editor/Source/InterchangeEditor/Private/InterchangeResetContextMenuExtender.h
@@ -0,0 +1,14 @@
+// Copyright Epic Games, Inc. All Rights Reserved.
+
+#pragma once
+#include "Misc/CoreDelegates.h"
+
+class FInterchangeResetContextMenuExtender
+{
+public:
+	static void SetupLevelEditorContextMenuExtender();
+	static void RemoveLevelEditorContextMenuExtender();
+
+private:
+	static FDelegateHandle LevelEditorExtenderDelegateHandle;
+};
diff --git a/Engine/Plugins/Interchange/Runtime/Source/FactoryNodes/Private/InterchangeActorFactoryNode.cpp b/Engine/Plugins/Interchange/Runtime/Source/FactoryNodes/Private/InterchangeActorFactoryNode.cpp
--- a/Engine/Plugins/Interchange/Runtime/Source/FactoryNodes/Private/InterchangeActorFactoryNode.cpp
+++ b/Engine/Plugins/Interchange/Runtime/Source/FactoryNodes/Private/InterchangeActorFactoryNode.cpp
@@ -152,4 +152,11 @@ bool UInterchangeActorFactoryNode::AddTags(const TArray<FString>& InTags)
 bool UInterchangeActorFactoryNode::RemoveTag(const FString& Tag)
 {
 	return Tags.RemoveItem(Tag);
-}
\ No newline at end of file
+}
+
+void UInterchangeActorFactoryNode::OnRestoreAllCustomAttributeDelegates()
+{
+	REFILL_CUSTOM_ATTRIBUTE_APPLY_DELEGATE(GlobalTransform, FTransform)
+	REFILL_CUSTOM_ATTRIBUTE_APPLY_DELEGATE(LocalTransform, FTransform)
+}
+
diff --git a/Engine/Plugins/Interchange/Runtime/Source/FactoryNodes/Public/InterchangeActorFactoryNode.h b/Engine/Plugins/Interchange/Runtime/Source/FactoryNodes/Public/InterchangeActorFactoryNode.h
--- a/Engine/Plugins/Interchange/Runtime/Source/FactoryNodes/Public/InterchangeActorFactoryNode.h
+++ b/Engine/Plugins/Interchange/Runtime/Source/FactoryNodes/Public/InterchangeActorFactoryNode.h
@@ -80,6 +80,8 @@ public:
 	/** Remove Tag that this Actor is supposed to have. */
 	bool RemoveTag(const FString& Tag);
 
+
+	virtual void OnRestoreAllCustomAttributeDelegates() override;
 private:
 
 	bool ApplyCustomGlobalTransformToAsset(UObject* Asset) const;
diff --git a/Engine/Plugins/Interchange/Runtime/Source/Import/Private/Scene/InterchangeActorFactory.cpp b/Engine/Plugins/Interchange/Runtime/Source/Import/Private/Scene/InterchangeActorFactory.cpp
--- a/Engine/Plugins/Interchange/Runtime/Source/Import/Private/Scene/InterchangeActorFactory.cpp
+++ b/Engine/Plugins/Interchange/Runtime/Source/Import/Private/Scene/InterchangeActorFactory.cpp
@@ -9,6 +9,7 @@
 
 #include "GameFramework/Actor.h"
 #include "Components/SceneComponent.h"
+#include "Components/ActorComponent.h"
 
 #if WITH_EDITORONLY_DATA
 #include "Engine/World.h"
@@ -21,67 +22,69 @@ extern UNREALED_API UEditorEngine* GEditor;
 
 #include UE_INLINE_GENERATED_CPP_BY_NAME(InterchangeActorFactory)
 
-namespace UE::InterchangeActorFactory::Private
+UClass* UInterchangeActorFactory::GetFactoryClass() const
 {
-	// For the camera actor types we get two components each: The root component is a default scene component, and the actual camera
-	// component is a child of the scene component. We want to place all scene component stuff (mostly transform) on the root
-	// component, and all the camera stuff on the camera component. This agrees with how the actor/root component is bound on
-	// LevelSequences, and is likely what users expect because when you place a camera actor on the level and move it around, you always
-	// affect the root component transform
-
-	template <class T>
-	void ApplyAllCameraCustomAttributes(
-		const UInterchangeFactoryBase::FImportSceneObjectsParams& CreateSceneObjectsParams,
-		T* CameraFactoryNode,
-		USceneComponent* RootSceneComponent,
-		USceneComponent* ChildCameraComponent
-	)
-	{
-		using namespace UE::Interchange;
-
-		if (!RootSceneComponent || !ChildCameraComponent)
-		{
-			return;
-		}
-
-		UInterchangeBaseNodeContainer* const NodeContainer = const_cast<UInterchangeBaseNodeContainer* const>(CreateSceneObjectsParams.NodeContainer);
+	return AActor::StaticClass();
+}
 
-		// Create a temp factory node so we don't modify our existing nodes with our changes
-		T* FactoryNodeCopy = NewObject<T>(NodeContainer, NAME_None);
-		FactoryNodeCopy->InitializeNode(
-			CameraFactoryNode->GetUniqueID(),
-			CameraFactoryNode->GetDisplayLabel(),
-			CameraFactoryNode->GetNodeContainerType()
-		);
+void UInterchangeActorFactory::ExecuteResetObjectProperties(const UInterchangeBaseNodeContainer* BaseNodeContainer, UInterchangeFactoryBaseNode* FactoryNode, UObject* ImportedObject)
+{
+	using namespace UE::Interchange;
 
-		NodeContainer->ReplaceNode(CameraFactoryNode->GetUniqueID(), FactoryNodeCopy);
+	FImportSceneObjectsParams TempSceneObjectParmeters;
+	TempSceneObjectParmeters.FactoryNode = FactoryNode;
+	TempSceneObjectParmeters.NodeContainer = BaseNodeContainer;
+	TempSceneObjectParmeters.ReimportObject = ImportedObject;
+	
+	if (AActor* ImportedActor = ActorHelper::SpawnFactoryActor(TempSceneObjectParmeters))
+	{
+		if (UInterchangeActorFactoryNode* ActorFactoryNode = Cast<UInterchangeActorFactoryNode>(FactoryNode))
 		{
-			UInterchangeFactoryBase::FImportSceneObjectsParams ParamsCopy{CreateSceneObjectsParams};
-			ParamsCopy.FactoryNode = FactoryNodeCopy;
-
-			// Apply exclusively camera stuff to the CineCameraComponent
-			FactoryNodeCopy->CopyWithObject(CameraFactoryNode, ChildCameraComponent);
-			FactoryNodeCopy->RemoveCustomAttributesForClass(USceneComponent::StaticClass());
-			ActorHelper::ApplyAllCustomAttributes(ParamsCopy, *ChildCameraComponent);
-
-			// Apply exclusively scene component stuff to the root SceneComponent
-			FactoryNodeCopy->CopyWithObject(CameraFactoryNode, RootSceneComponent);
-			FactoryNodeCopy->RemoveCustomAttributesForClass(UCineCameraComponent::StaticClass());
-			ActorHelper::ApplyAllCustomAttributes(ParamsCopy, *RootSceneComponent);
+			if (UObject* ObjectToUpdate = ProcessActor(*ImportedActor, *ActorFactoryNode, *BaseNodeContainer, TempSceneObjectParmeters))
+			{
+				if (USceneComponent* RootComponent = ImportedActor->GetRootComponent())
+				{
+					UActorComponent* ActorComponent = Cast<UActorComponent>(ObjectToUpdate);
+					if (ActorComponent)
+					{
+#if WITH_EDITOR
+						ActorComponent->PreEditChange(nullptr);
+#endif
+						ActorComponent->UnregisterComponent();
+					}
+
+					// Cache mobility value to allow application of transform
+					EComponentMobility::Type CachedMobility = RootComponent->Mobility;
+					RootComponent->SetMobility(EComponentMobility::Type::Movable);
+					
+					ApplyAllCustomAttributesToObject(TempSceneObjectParmeters, *ImportedActor, ObjectToUpdate);
+
+					// Restore mobility value
+					if (CachedMobility != EComponentMobility::Type::Movable)
+					{
+						RootComponent->SetMobility(CachedMobility);
+					}
+					
+					if (ActorComponent)
+					{
+						ActorComponent->RegisterComponent();
+#if WITH_EDITOR
+						ActorComponent->PostEditChange();
+#endif
+					}
+					
+					return;
+				}
+			}
 		}
-		NodeContainer->ReplaceNode(FactoryNodeCopy->GetUniqueID(), CameraFactoryNode);
 	}
-}
 
-UClass* UInterchangeActorFactory::GetFactoryClass() const
-{
-	return AActor::StaticClass();
+	Super::ExecuteResetObjectProperties(BaseNodeContainer, FactoryNode, ImportedObject);	
 }
 
 UObject* UInterchangeActorFactory::ImportSceneObject_GameThread(const UInterchangeFactoryBase::FImportSceneObjectsParams& CreateSceneObjectsParams)
 {
 	using namespace UE::Interchange;
-	using namespace UE::InterchangeActorFactory::Private;
 
 	UInterchangeActorFactoryNode* FactoryNode = Cast<UInterchangeActorFactoryNode>(CreateSceneObjectsParams.FactoryNode);
 	if (!ensure(FactoryNode) || !CreateSceneObjectsParams.NodeContainer)
@@ -101,21 +104,7 @@ UObject* UInterchangeActorFactory::ImportSceneObject_GameThread(const UInterchan
 				EComponentMobility::Type CachedMobility = RootComponent->Mobility;
 				RootComponent->SetMobility(EComponentMobility::Type::Movable);
 
-				// Apply factory node to object(s)
-				if (FactoryNode->IsA<UInterchangePhysicalCameraFactoryNode>())
-				{
-					UCineCameraComponent* CameraComponent = Cast<UCineCameraComponent>(ObjectToUpdate);
-					ApplyAllCameraCustomAttributes(CreateSceneObjectsParams, Cast<UInterchangePhysicalCameraFactoryNode>(FactoryNode), RootComponent, CameraComponent);
-				}
-				else if (FactoryNode->IsA<UInterchangeStandardCameraFactoryNode>())
-				{
-					UCameraComponent* CameraComponent = Cast<UCameraComponent>(ObjectToUpdate);
-					ApplyAllCameraCustomAttributes(CreateSceneObjectsParams, Cast<UInterchangeStandardCameraFactoryNode>(FactoryNode), RootComponent, CameraComponent);
-				}
-				else
-				{
-					ActorHelper::ApplyAllCustomAttributes(CreateSceneObjectsParams, *ObjectToUpdate);
-				}
+				ApplyAllCustomAttributesToObject(CreateSceneObjectsParams, *SpawnedActor, ObjectToUpdate);
 
 				// Restore mobility value
 				if (CachedMobility != EComponentMobility::Type::Movable)
@@ -133,7 +122,7 @@ UObject* UInterchangeActorFactory::ImportSceneObject_GameThread(const UInterchan
 	return SpawnedActor;
 }
 
-UObject* UInterchangeActorFactory::ProcessActor(AActor& SpawnedActor, const UInterchangeActorFactoryNode& /*FactoryNode*/, const UInterchangeBaseNodeContainer& /*NodeContainer*/, const FImportSceneObjectsParams& Params)
+UObject* UInterchangeActorFactory::ProcessActor(AActor& SpawnedActor, const UInterchangeActorFactoryNode& /*FactoryNode*/, const UInterchangeBaseNodeContainer& /*NodeContainer*/, const FImportSceneObjectsParams& /*Params*/)
 {
 	return SpawnedActor.GetRootComponent();
 }
@@ -227,4 +216,10 @@ void UInterchangeActorFactory::AddUniqueLayersToWorld(UWorld* World, const TSet<
 		World->Modify();
 	}
 }
-#endif
\ No newline at end of file
+#endif
+
+void UInterchangeActorFactory::ApplyAllCustomAttributesToObject(const UInterchangeFactoryBase::FImportSceneObjectsParams& CreateSceneObjectsParams, AActor& SpawnedActor, UObject* ObjectToUpdate)
+{
+	using namespace UE::Interchange;
+	ActorHelper::ApplyAllCustomAttributes(CreateSceneObjectsParams, *ObjectToUpdate);
+}
diff --git a/Engine/Plugins/Interchange/Runtime/Source/Import/Private/Scene/InterchangeCameraActorFactory.cpp b/Engine/Plugins/Interchange/Runtime/Source/Import/Private/Scene/InterchangeCameraActorFactory.cpp
--- a/Engine/Plugins/Interchange/Runtime/Source/Import/Private/Scene/InterchangeCameraActorFactory.cpp
+++ b/Engine/Plugins/Interchange/Runtime/Source/Import/Private/Scene/InterchangeCameraActorFactory.cpp
@@ -5,6 +5,7 @@
 #include "InterchangeCameraFactoryNode.h"
 #include "Scene/InterchangeActorFactory.h"
 #include "Scene/InterchangeActorHelper.h"
+#include "Nodes/InterchangeBaseNodeContainer.h"
 
 #include "CineCameraActor.h"
 #include "CineCameraComponent.h"
@@ -14,6 +15,59 @@
 
 #include UE_INLINE_GENERATED_CPP_BY_NAME(InterchangeCameraActorFactory)
 
+namespace UE::InterchangeCameraActorFactory::Private
+{
+	// For the camera actor types we get two components each: The root component is a default scene component, and the actual camera
+	// component is a child of the scene component. We want to place all scene component stuff (mostly transform) on the root
+	// component, and all the camera stuff on the camera component. This agrees with how the actor/root component is bound on
+	// LevelSequences, and is likely what users expect because when you place a camera actor on the level and move it around, you always
+	// affect the root component transform
+
+	template <class T>
+	void ApplyAllCameraCustomAttributes(
+		const UInterchangeFactoryBase::FImportSceneObjectsParams& CreateSceneObjectsParams,
+		T* CameraFactoryNode,
+		USceneComponent* RootSceneComponent,
+		USceneComponent* ChildCameraComponent
+	)
+	{
+		using namespace UE::Interchange;
+
+		if (!RootSceneComponent || !ChildCameraComponent)
+		{
+			return;
+		}
+
+		UInterchangeBaseNodeContainer* const NodeContainer = const_cast<UInterchangeBaseNodeContainer* const>(CreateSceneObjectsParams.NodeContainer);
+
+		// Create a temp factory node so we don't modify our existing nodes with our changes
+		T* FactoryNodeCopy = NewObject<T>(NodeContainer, NAME_None);
+		FactoryNodeCopy->InitializeNode(
+			CameraFactoryNode->GetUniqueID(),
+			CameraFactoryNode->GetDisplayLabel(),
+			CameraFactoryNode->GetNodeContainerType()
+		);
+
+		NodeContainer->ReplaceNode(CameraFactoryNode->GetUniqueID(), FactoryNodeCopy);
+		{
+			UInterchangeFactoryBase::FImportSceneObjectsParams ParamsCopy{ CreateSceneObjectsParams };
+			ParamsCopy.FactoryNode = FactoryNodeCopy;
+
+			// Apply exclusively camera stuff to the CineCameraComponent
+			FactoryNodeCopy->CopyWithObject(CameraFactoryNode, ChildCameraComponent);
+			FactoryNodeCopy->RemoveCustomAttributesForClass(USceneComponent::StaticClass());
+			ActorHelper::ApplyAllCustomAttributes(ParamsCopy, *ChildCameraComponent);
+
+			// Apply exclusively scene component stuff to the root SceneComponent
+			FactoryNodeCopy->CopyWithObject(CameraFactoryNode, RootSceneComponent);
+			FactoryNodeCopy->RemoveCustomAttributesForClass(UCineCameraComponent::StaticClass());
+			ActorHelper::ApplyAllCustomAttributes(ParamsCopy, *RootSceneComponent);
+		}
+		NodeContainer->ReplaceNode(FactoryNodeCopy->GetUniqueID(), CameraFactoryNode);
+	}
+}
+
+
 UClass* UInterchangeCineCameraActorFactory::GetFactoryClass() const
 {
 	return ACineCameraActor::StaticClass();
@@ -26,6 +80,25 @@ UObject* UInterchangeCineCameraActorFactory::ProcessActor(AActor& SpawnedActor,
 	return CineCameraActor ? CineCameraActor->GetCineCameraComponent() : nullptr;
 }
 
+void UInterchangeCineCameraActorFactory::ApplyAllCustomAttributesToObject(const UInterchangeFactoryBase::FImportSceneObjectsParams& CreateSceneObjectsParams, AActor& SpawnedActor, UObject* ObjectToUpdate)
+{
+	using namespace UE::InterchangeCameraActorFactory::Private;
+
+	if (UCineCameraComponent* CameraComponent = Cast<UCineCameraComponent>(ObjectToUpdate))
+	{
+		if (UInterchangePhysicalCameraFactoryNode* FactoryNode = Cast<UInterchangePhysicalCameraFactoryNode>(CreateSceneObjectsParams.FactoryNode))
+		{
+			if (USceneComponent* RootComponent = SpawnedActor.GetRootComponent())
+			{
+				ApplyAllCameraCustomAttributes(CreateSceneObjectsParams, FactoryNode, RootComponent, CameraComponent);
+				return;
+			}
+		}
+	}
+
+	Super::ApplyAllCustomAttributesToObject(CreateSceneObjectsParams, SpawnedActor, ObjectToUpdate);
+}
+
 
 UClass* UInterchangeCameraActorFactory::GetFactoryClass() const
 {
@@ -37,4 +110,23 @@ UObject* UInterchangeCameraActorFactory::ProcessActor(AActor& SpawnedActor, cons
 	ACameraActor* CameraActor = Cast<ACameraActor>(&SpawnedActor);
 
 	return CameraActor ? CameraActor->GetCameraComponent() : nullptr;
-}
\ No newline at end of file
+}
+
+void UInterchangeCameraActorFactory::ApplyAllCustomAttributesToObject(const UInterchangeFactoryBase::FImportSceneObjectsParams& CreateSceneObjectsParams, AActor& SpawnedActor, UObject* ObjectToUpdate)
+{
+	using namespace UE::InterchangeCameraActorFactory::Private;
+
+	if (UCameraComponent* CameraComponent = Cast<UCameraComponent>(ObjectToUpdate))
+	{
+		if (UInterchangeStandardCameraFactoryNode* FactoryNode = Cast<UInterchangeStandardCameraFactoryNode>(CreateSceneObjectsParams.FactoryNode))
+		{
+			if (USceneComponent* RootComponent = SpawnedActor.GetRootComponent())
+			{
+				ApplyAllCameraCustomAttributes(CreateSceneObjectsParams, FactoryNode, RootComponent, CameraComponent);
+				return;
+			}
+		}
+	}
+
+	Super::ApplyAllCustomAttributesToObject(CreateSceneObjectsParams, SpawnedActor, ObjectToUpdate);
+}
diff --git a/Engine/Plugins/Interchange/Runtime/Source/Import/Private/Scene/InterchangeStaticMeshActorFactory.cpp b/Engine/Plugins/Interchange/Runtime/Source/Import/Private/Scene/InterchangeStaticMeshActorFactory.cpp
--- a/Engine/Plugins/Interchange/Runtime/Source/Import/Private/Scene/InterchangeStaticMeshActorFactory.cpp
+++ b/Engine/Plugins/Interchange/Runtime/Source/Import/Private/Scene/InterchangeStaticMeshActorFactory.cpp
@@ -127,4 +127,23 @@ UObject* UInterchangeStaticMeshActorFactory::ProcessActor(AActor& SpawnedActor,
 	}
 
 	return nullptr;
-};
+}
+
+void UInterchangeStaticMeshActorFactory::ExecuteResetObjectProperties(const UInterchangeBaseNodeContainer* BaseNodeContainer, UInterchangeFactoryBaseNode* FactoryNode, UObject* ImportedObject)
+{
+	Super::ExecuteResetObjectProperties(BaseNodeContainer, FactoryNode, ImportedObject);
+	if (AActor* SpawnedActor = Cast<AActor>(ImportedObject))
+	{
+		if (UInterchangeMeshActorFactoryNode* MeshActorFactoryNode = Cast<UInterchangeMeshActorFactoryNode>(FactoryNode))
+		{
+			if (UMeshComponent* MeshComponent = Cast<UMeshComponent>(SpawnedActor->GetRootComponent()))
+			{
+				using namespace UE::Interchange;
+				
+				MeshComponent->EmptyOverrideMaterials();
+				ActorHelper::ApplySlotMaterialDependencies(*BaseNodeContainer, *MeshActorFactoryNode, *MeshComponent);
+			}
+		}
+	}
+}
+
diff --git a/Engine/Plugins/Interchange/Runtime/Source/Import/Public/InterchangeAssetUserData.h b/Engine/Plugins/Interchange/Runtime/Source/Import/Public/InterchangeAssetUserData.h
--- a/Engine/Plugins/Interchange/Runtime/Source/Import/Public/InterchangeAssetUserData.h
+++ b/Engine/Plugins/Interchange/Runtime/Source/Import/Public/InterchangeAssetUserData.h
@@ -16,4 +16,17 @@ class INTERCHANGEIMPORT_API UInterchangeAssetUserData : public UAssetUserData
 public:
 	UPROPERTY(VisibleAnywhere, BlueprintReadWrite, AdvancedDisplay, Category = "Interchange User Data", meta = (ScriptName = "Metadata", DisplayName = "Metadata"))
 	TMap<FString, FString> MetaData;
+};
+
+/** Asset user data that can be used with Interchange on Levels. This will reside in the world settings of the ULevel*/
+UCLASS(BlueprintType, meta = (ScriptName = "InterchangeLevelUserData", DisplayName = "Interchange Level User Data"))
+class INTERCHANGEIMPORT_API UInterchangeLevelAssetUserData : public UAssetUserData
+{
+	GENERATED_BODY()
+
+public:
+#if WITH_EDITORONLY_DATA
+	UPROPERTY(VisibleAnywhere, BlueprintReadWrite, AdvancedDisplay, Category = "Interchange User Data", meta = (ScriptName = "SceneImportPaths", DisplayName = "Scene Import Paths"))
+	TArray<FSoftObjectPath> SceneImportPaths;
+#endif
 };
\ No newline at end of file
diff --git a/Engine/Plugins/Interchange/Runtime/Source/Import/Public/Scene/InterchangeActorFactory.h b/Engine/Plugins/Interchange/Runtime/Source/Import/Public/Scene/InterchangeActorFactory.h
--- a/Engine/Plugins/Interchange/Runtime/Source/Import/Public/Scene/InterchangeActorFactory.h
+++ b/Engine/Plugins/Interchange/Runtime/Source/Import/Public/Scene/InterchangeActorFactory.h
@@ -26,6 +26,9 @@ public:
 
 	virtual UClass* GetFactoryClass() const override;
 
+protected:
+	virtual void ExecuteResetObjectProperties(const UInterchangeBaseNodeContainer* BaseNodeContainer, UInterchangeFactoryBaseNode* FactoryNode, UObject* ImportedObject) override;
+
 private:
 	virtual UObject* ImportSceneObject_GameThread(const UInterchangeFactoryBase::FImportSceneObjectsParams& CreateSceneObjectsParams) override;
 
@@ -46,6 +49,13 @@ protected:
 	 */
 	virtual UObject* ProcessActor(AActor& SpawnedActor, const UInterchangeActorFactoryNode& FactoryNode, const UInterchangeBaseNodeContainer& NodeContainer, const FImportSceneObjectsParams& Params);
 
+
+	/**
+	 * Method used to apply all the custom attributes to the Actor or its sub components.
+	 * Base class method will simply use ActorHelper to apply all the custom attributes to the object.
+	 */
+	virtual void ApplyAllCustomAttributesToObject(const UInterchangeFactoryBase::FImportSceneObjectsParams& CreateSceneObjectsParams, AActor& SpawnedActor, UObject* ObjectToUpdate);
+
 	template <typename T>
 	void LogMessage(const FImportSceneObjectsParams& Params, const FText& Message, const FString& ActorLabel)
 	{
diff --git a/Engine/Plugins/Interchange/Runtime/Source/Import/Public/Scene/InterchangeCameraActorFactory.h b/Engine/Plugins/Interchange/Runtime/Source/Import/Public/Scene/InterchangeCameraActorFactory.h
--- a/Engine/Plugins/Interchange/Runtime/Source/Import/Public/Scene/InterchangeCameraActorFactory.h
+++ b/Engine/Plugins/Interchange/Runtime/Source/Import/Public/Scene/InterchangeCameraActorFactory.h
@@ -28,6 +28,7 @@ public:
 
 	virtual UObject* ProcessActor(class AActor& SpawnedActor, const UInterchangeActorFactoryNode& FactoryNode, const UInterchangeBaseNodeContainer& NodeContainer, const FImportSceneObjectsParams& Params) override;
 
+	virtual void ApplyAllCustomAttributesToObject(const UInterchangeFactoryBase::FImportSceneObjectsParams& CreateSceneObjectsParams, AActor& SpawnedActor, UObject* ObjectToUpdate) override;
 
 	// Interchange actor factory interface end
 	//////////////////////////////////////////////////////////////////////////
@@ -52,7 +53,7 @@ public:
 
 	virtual UObject* ProcessActor(class AActor& SpawnedActor, const UInterchangeActorFactoryNode& FactoryNode, const UInterchangeBaseNodeContainer& NodeContainer, const FImportSceneObjectsParams& Params) override;
 
-
+	virtual void ApplyAllCustomAttributesToObject(const UInterchangeFactoryBase::FImportSceneObjectsParams& CreateSceneObjectsParams, AActor& SpawnedActor, UObject* ObjectToUpdate) override;
 	// Interchange actor factory interface end
 	//////////////////////////////////////////////////////////////////////////
 };
diff --git a/Engine/Plugins/Interchange/Runtime/Source/Import/Public/Scene/InterchangeStaticMeshActorFactory.h b/Engine/Plugins/Interchange/Runtime/Source/Import/Public/Scene/InterchangeStaticMeshActorFactory.h
--- a/Engine/Plugins/Interchange/Runtime/Source/Import/Public/Scene/InterchangeStaticMeshActorFactory.h
+++ b/Engine/Plugins/Interchange/Runtime/Source/Import/Public/Scene/InterchangeStaticMeshActorFactory.h
@@ -33,6 +33,8 @@ protected:
 
 	virtual UObject* ProcessActor(class AActor& SpawnedActor, const UInterchangeActorFactoryNode& FactoryNode, const UInterchangeBaseNodeContainer& NodeContainer, const FImportSceneObjectsParams& Params) override;
 
+
+	virtual void ExecuteResetObjectProperties(const UInterchangeBaseNodeContainer* BaseNodeContainer, UInterchangeFactoryBaseNode* FactoryNode, UObject* ImportedObject) override;
 	// Interchange actor factory interface end
 	//////////////////////////////////////////////////////////////////////////
 };
diff --git a/Engine/Plugins/Interchange/Runtime/Source/Pipelines/Private/InterchangeGenericScenesPipeline.cpp b/Engine/Plugins/Interchange/Runtime/Source/Pipelines/Private/InterchangeGenericScenesPipeline.cpp
--- a/Engine/Plugins/Interchange/Runtime/Source/Pipelines/Private/InterchangeGenericScenesPipeline.cpp
+++ b/Engine/Plugins/Interchange/Runtime/Source/Pipelines/Private/InterchangeGenericScenesPipeline.cpp
@@ -2,22 +2,23 @@
 
 #include "InterchangeGenericScenesPipeline.h"
 
-#include "Nodes/InterchangeUserDefinedAttribute.h"
 #include "InterchangeActorFactoryNode.h"
 #include "InterchangeAssetImportData.h"
-#include "InterchangeCameraNode.h"
+#include "InterchangeAssetUserData.h"
 #include "InterchangeCameraFactoryNode.h"
+#include "InterchangeCameraNode.h"
 #include "InterchangeCommonPipelineDataFactoryNode.h"
-#include "InterchangeLevelInstanceActorFactoryNode.h"
-#include "InterchangeLightNode.h"
-#include "InterchangeLightFactoryNode.h"
-#include "InterchangeMeshActorFactoryNode.h"
-#include "InterchangeMeshNode.h"
 #include "InterchangeDecalActorFactoryNode.h"
 #include "InterchangeDecalNode.h"
 #include "InterchangeEditorUtilitiesBase.h"
+#include "InterchangeImportReset.h"
 #include "InterchangeLevelFactoryNode.h"
+#include "InterchangeLevelInstanceActorFactoryNode.h"
+#include "InterchangeLightFactoryNode.h"
+#include "InterchangeLightNode.h"
 #include "InterchangeManager.h"
+#include "InterchangeMeshActorFactoryNode.h"
+#include "InterchangeMeshNode.h"
 #include "InterchangePipelineLog.h"
 #include "InterchangePipelineMeshesUtilities.h"
 #include "InterchangeSceneImportAsset.h"
@@ -28,6 +29,7 @@
 #include "InterchangeSkeletalMeshLodDataNode.h"
 #include "InterchangeSkeletonFactoryNode.h"
 #include "InterchangeVariantSetNode.h"
+#include "Nodes/InterchangeUserDefinedAttribute.h"
 
 
 #include "Animation/SkeletalMeshActor.h"
@@ -39,9 +41,10 @@
 #include "Engine/SpotLight.h"
 #include "Engine/StaticMeshActor.h"
 #include "Engine/World.h"
-#include "WorldPartition/WorldPartition.h"
+#include "GameFramework/WorldSettings.h"
 #include "LevelInstance/LevelInstanceActor.h"
 #include "Misc/PackageName.h"
+#include "WorldPartition/WorldPartition.h"
 
 #if WITH_EDITOR
 #include "ObjectTools.h"
@@ -1049,9 +1052,65 @@ void UInterchangeGenericLevelPipeline::ExecutePostImportPipeline(const UIntercha
 
 	if (UInterchangeSceneImportAsset* SceneImportAsset = Cast<UInterchangeSceneImportAsset>(CreatedAsset))
 	{
+		const FSoftObjectPath SceneImportAssetPath(SceneImportAsset);
+
+		auto AddAssetUserDataToObjects = [&]()
+		{
+			if (SceneHierarchyType == EInterchangeSceneHierarchyType::CreateLevelActors)
+			{
+				InBaseNodeContainer->IterateNodesOfType<UInterchangeActorFactoryNode>(
+					[&SceneImportAssetPath](const FString& NodeUid, UInterchangeActorFactoryNode* FactoryNode)
+					{
+						FSoftObjectPath ObjectPath;
+						if (FactoryNode->GetCustomReferenceObject(ObjectPath))
+						{
+							if (AActor* ReferencedObject = Cast<AActor>(ObjectPath.TryLoad()))
+							{
+								if (UWorld* World = ReferencedObject->GetWorld())
+								{
+									if (USceneComponent* RootComponent = ReferencedObject->GetRootComponent())
+									{
+										UInterchangeAssetUserData* AssetUserData = NewObject<UInterchangeAssetUserData>(ReferencedObject);
+										AssetUserData->MetaData.Add(UE::Interchange::InterchangeReset::Constants::SceneImportAssetPathKey, SceneImportAssetPath.ToString());
+										RootComponent->AddAssetUserData(AssetUserData);
+									}
+								}
+							}
+						}
+					}
+				);
+			}
+			else if (SceneHierarchyType == EInterchangeSceneHierarchyType::CreateLevelInstanceActor)
+			{
+				InBaseNodeContainer->IterateNodesOfType<UInterchangeLevelFactoryNode>(
+					[&SceneImportAssetPath](const FString& NodeUid, UInterchangeLevelFactoryNode* FactoryNode)
+					{
+						bool bWasLevelCreatedForImport = false;
+						if (FactoryNode->GetCustomShouldCreateLevel(bWasLevelCreatedForImport) && bWasLevelCreatedForImport)
+						{
+							FSoftObjectPath ObjectPath;
+							if (FactoryNode->GetCustomReferenceObject(ObjectPath))
+							{
+								if (UWorld* World = Cast<UWorld>(ObjectPath.TryLoad()))
+								{
+									if (AWorldSettings* WorldSettings = World->GetWorldSettings())
+									{
+										UInterchangeLevelAssetUserData* AssetUserData = NewObject<UInterchangeLevelAssetUserData>(World);
+										AssetUserData->SceneImportPaths.Add(SceneImportAssetPath);
+										WorldSettings->AddAssetUserData(AssetUserData);
+									}
+								}
+							}
+						}
+					}
+				);
+			}
+		};
+
 		if (!bIsReimportContext || !(bDeleteMissingActors || bDeleteMissingAssets))
 		{
 			SceneImportAsset->UpdateSceneObjects();
+			AddAssetUserDataToObjects();
 			return;
 		}
 
@@ -1059,6 +1118,7 @@ void UInterchangeGenericLevelPipeline::ExecutePostImportPipeline(const UIntercha
 		if (!ensure(FactoryNode))
 		{
 			SceneImportAsset->UpdateSceneObjects();
+			AddAssetUserDataToObjects();
 			return;
 		}
 
@@ -1067,6 +1127,7 @@ void UInterchangeGenericLevelPipeline::ExecutePostImportPipeline(const UIntercha
 		SceneImportAsset->GetSceneSoftObjectPaths(PrevSoftObjectPaths);
 
 		SceneImportAsset->UpdateSceneObjects();
+		AddAssetUserDataToObjects();
 
 		// Nothing to take care of
 		if (PrevSoftObjectPaths.IsEmpty())
diff --git a/Engine/Source/Runtime/Interchange/Core/Private/Nodes/InterchangeFactoryBaseNode.cpp b/Engine/Source/Runtime/Interchange/Core/Private/Nodes/InterchangeFactoryBaseNode.cpp
--- a/Engine/Source/Runtime/Interchange/Core/Private/Nodes/InterchangeFactoryBaseNode.cpp
+++ b/Engine/Source/Runtime/Interchange/Core/Private/Nodes/InterchangeFactoryBaseNode.cpp
@@ -21,6 +21,12 @@ namespace UE
 			return BaseNodeFactoryDependencies_BaseKey;
 		}
 
+		const FString& FFactoryBaseNodeStaticData::AttributesAppliedThroughDelegatesBaseKey()
+		{
+			static FString AttributesAppliedThroughDelegates_BaseKey = TEXT("__AttributeAppliedThroughDelegate__");
+			return AttributesAppliedThroughDelegates_BaseKey;
+		}
+
 		const FAttributeKey& FFactoryBaseNodeStaticData::ReimportStrategyFlagsKey()
 		{
 			static FAttributeKey AttributeKey(TEXT("__Reimport_Strategy_Flags_Key__"));
@@ -46,6 +52,7 @@ namespace UE
 UInterchangeFactoryBaseNode::UInterchangeFactoryBaseNode()
 {
 	FactoryDependencies.Initialize(Attributes, UE::Interchange::FFactoryBaseNodeStaticData::FactoryDependenciesBaseKey());
+	AttributesAppliedThroughDelegates.Initialize(Attributes, UE::Interchange::FFactoryBaseNodeStaticData::AttributesAppliedThroughDelegatesBaseKey());
 
 	RegisterAttribute<uint8>(UE::Interchange::FFactoryBaseNodeStaticData::ReimportStrategyFlagsKey(), static_cast<uint8>(EReimportStrategyFlags::ApplyNoProperties));
 }
@@ -299,12 +306,30 @@ void UInterchangeFactoryBaseNode::FillAllCustomAttributeFromObject(UObject* Obje
 	}
 }
 
+void UInterchangeFactoryBaseNode::RestoreAllCustomAttributeDelegates()
+{
+	TArray<FString> AttributeKeys;
+	AttributesAppliedThroughDelegates.GetItems(AttributeKeys);
+	for (const FString& Key : AttributeKeys)
+	{
+		AttributesAppliedThroughDelegatesKeySet.Add(Key);
+	}
+	
+	OnRestoreAllCustomAttributeDelegates();
+}
+
 void UInterchangeFactoryBaseNode::RemoveCustomAttributesForClass(UClass* Class)
 {
 	ApplyCustomAttributeDelegates.Remove(Class);
 	FillCustomAttributeDelegates.Remove(Class);
 }
 
+void UInterchangeFactoryBaseNode::RemoveAllCustomAttributeDelegates()
+{
+	ApplyCustomAttributeDelegates.Empty();
+	FillCustomAttributeDelegates.Empty();
+}
+
 UInterchangeFactoryBaseNode* UInterchangeFactoryBaseNode::DuplicateWithObject(const UInterchangeFactoryBaseNode* SourceNode, UObject* Object)
 {
 	UInterchangeFactoryBaseNode* CurrentNode = NewObject<UInterchangeFactoryBaseNode>(GetTransientPackage(), SourceNode->GetClass());
@@ -352,3 +377,8 @@ bool UInterchangeFactoryBaseNode::UnsetForceNodeReimport()
 	}
 	return false;
 }
+
+void UInterchangeFactoryBaseNode::OnRestoreAllCustomAttributeDelegates()
+{
+	// Should be filled with appropriate attributes in the child classes.
+}
diff --git a/Engine/Source/Runtime/Interchange/Core/Public/InterchangeFactoryBase.h b/Engine/Source/Runtime/Interchange/Core/Public/InterchangeFactoryBase.h
--- a/Engine/Source/Runtime/Interchange/Core/Public/InterchangeFactoryBase.h
+++ b/Engine/Source/Runtime/Interchange/Core/Public/InterchangeFactoryBase.h
@@ -221,7 +221,7 @@ public:
 		 * If not null, the factory must perform a reimport of the scene node.
 		 */
 		UObject* ReimportObject = nullptr;
-	
+
 		/**
 		 * Factory base node associated with the reimported scene node.
 		 */
@@ -266,7 +266,6 @@ public:
 		UInterchangeTranslatorBase* Translator;
 		
 		bool bIsReimport  = false;
-;
 	};
 
 	/*
@@ -277,7 +276,7 @@ public:
 	{
 		check(IsInGameThread());
 	}
-	
+
 	/*
 	 * Build the asset if it can be built.
 	 * 
@@ -286,7 +285,7 @@ public:
 	 * @note - This function is called when starting the pre-completion task.
 	 * @note - The default implementation for asset that do not have any build step is to call PostEditChange and set OutPostEditchangeCalled to true.
 	 */
-	virtual void BuildObject_GameThread(const FSetupObjectParams& Arguments, bool &OutPostEditchangeCalled)
+	virtual void BuildObject_GameThread(const FSetupObjectParams& Arguments, bool& OutPostEditchangeCalled)
 	{
 		check(IsInGameThread());
 #if WITH_EDITOR
@@ -370,16 +369,50 @@ public:
 		check(Results != nullptr);
 		Results->Add(Item);
 	}
-	
+
 
 	void SetResultsContainer(UInterchangeResultsContainer* InResults)
 	{
 		Results = InResults;
 	}
 
+	/**
+	 * Called before resetting object properties. 
+	 * Restores all the tracked custom attribute delegates.
+	 */
+	void PreResetObjectProperties(const UInterchangeBaseNodeContainer* BaseNodeContainer, UInterchangeFactoryBaseNode* FactoryNode, UObject* ImportedObject)
+	{
+		// Restore delegates and force make the re-import strategy to apply pipeline properties.
+		FactoryNode->RestoreAllCustomAttributeDelegates();
+		
+		ExecutePreResetObjectProperties(BaseNodeContainer, FactoryNode, ImportedObject);
+	}
 
-	UPROPERTY()
-	TObjectPtr<UInterchangeResultsContainer> Results;
+	/**
+	 * Re-apply all the custom attributes that were set on the Factory Node to the Object.
+	 * Useful when resetting the object to its original imported state, unless the factory node's attributes have been changed.
+	 */
+	void ResetObjectProperties(const UInterchangeBaseNodeContainer* BaseNodeContainer, UInterchangeFactoryBaseNode* FactoryNode, UObject* ImportedObject)
+	{
+		// Force use of ApplyPipelineProperties Reimport Strategy to apply attributes through the attribute delegates.
+		EReimportStrategyFlags CachedReimportStrategyFlags = FactoryNode->GetReimportStrategyFlags();
+		FactoryNode->SetReimportStrategyFlags(EReimportStrategyFlags::ApplyPipelineProperties);
+		
+		ExecuteResetObjectProperties(BaseNodeContainer, FactoryNode, ImportedObject);
+		
+		// Restore Previous State
+		FactoryNode->SetReimportStrategyFlags(CachedReimportStrategyFlags);
+	}
+
+	/**
+	 * Called after resetting object properties. 
+	 * Useful for doing any cleanup required.
+	 */
+	void PostResetObjectProperties(const UInterchangeBaseNodeContainer* BaseNodeContainer, UInterchangeFactoryBaseNode* FactoryNode, UObject* ImportedObject)
+	{
+		ExecutePostResetObjectProperties(BaseNodeContainer, FactoryNode, ImportedObject);		
+		FactoryNode->RemoveAllCustomAttributeDelegates();
+	}
 
 	/**
 	* Acquires the Object to be re-imported.
@@ -391,4 +424,33 @@ public:
 	{
 		return ReimportObject;
 	}
+
+protected:
+
+	/**
+	 * Override this method in the child classes, to support custom behavior for PreResetObjectProperties
+	 */
+	virtual void ExecutePreResetObjectProperties(const UInterchangeBaseNodeContainer* BaseNodeContainer, UInterchangeFactoryBaseNode* FactoryNode, UObject* ImportedObject)
+	{
+	}
+
+	/**
+	 * Override this method in the child classes if the Custom Attributes are not directly applied to the imported object
+	 * E.g. Actor Factories need root component/other components in the hierarchy. Actor Factory overrides this method for custom behavior
+	 */
+	virtual void ExecuteResetObjectProperties(const UInterchangeBaseNodeContainer* BaseNodeContainer, UInterchangeFactoryBaseNode* FactoryNode, UObject* ImportedObject)
+	{
+		FactoryNode->ApplyAllCustomAttributeToObject(ImportedObject);
+	}
+
+	/**
+	 * Override this method in the child classes, to support custom behavior for PostResetObjectProperties
+	 */
+	virtual void ExecutePostResetObjectProperties(const UInterchangeBaseNodeContainer* BaseNodeContainer, UInterchangeFactoryBaseNode* FactoryNode, UObject* ImportedObject)
+	{
+	}
+
+public:
+	UPROPERTY()
+	TObjectPtr<UInterchangeResultsContainer> Results;
 };
diff --git a/Engine/Source/Runtime/Interchange/Core/Public/Nodes/InterchangeFactoryBaseNode.h b/Engine/Source/Runtime/Interchange/Core/Public/Nodes/InterchangeFactoryBaseNode.h
--- a/Engine/Source/Runtime/Interchange/Core/Public/Nodes/InterchangeFactoryBaseNode.h
+++ b/Engine/Source/Runtime/Interchange/Core/Public/Nodes/InterchangeFactoryBaseNode.h
@@ -38,6 +38,7 @@ struct FFrame;
 		if(bAddApplyDelegate)																														\
 		{																																			\
 			FName AttributeNameImpl = TEXT(""#AttributeName);																						\
+			AttributesAppliedThroughDelegates.AddItem(Macro_Custom##AttributeName##Key.Key);														\
 			AddApplyAndFillDelegates<AttributeType>(Macro_Custom##AttributeName##Key.Key, AssetType::StaticClass(), AttributeNameImpl);				\
 		}																																			\
 		return true;																																\
@@ -60,6 +61,7 @@ struct FFrame;
 		if(bAddApplyDelegate)																														\
 		{																																			\
 			check(NodeClassName::StaticClass()->IsChildOf<UInterchangeFactoryBaseNode>());															\
+			AttributesAppliedThroughDelegates.AddItem(Macro_Custom##AttributeName##Key.Key);														\
 			UClass* AssetClass = AssetType::StaticClass();																							\
 			TArray<UE::Interchange::FApplyAttributeToAsset>& Delegates = ApplyCustomAttributeDelegates.FindOrAdd(AssetClass);						\
 			Delegates.Add(UE::Interchange::FApplyAttributeToAsset::CreateUObject(this, &NodeClassName::ApplyCustom##AttributeName##ToAsset));		\
@@ -81,6 +83,7 @@ struct FFrame;
 		if(bAddApplyDelegate)																														\
 		{																																			\
 			check(NodeClassName::StaticClass()->IsChildOf<UInterchangeFactoryBaseNode>());															\
+			AttributesAppliedThroughDelegates.AddItem(Macro_Custom##AttributeName##Key.Key);														\
 			TArray<UE::Interchange::FApplyAttributeToAsset>& Delegates = ApplyCustomAttributeDelegates.FindOrAdd(GetObjectClass());					\
 			Delegates.Add(UE::Interchange::FApplyAttributeToAsset::CreateUObject(this, &NodeClassName::ApplyCustom##AttributeName##ToAsset));		\
 			TArray<UE::Interchange::FFillAttributeToAsset>& FillDelegates = FillCustomAttributeDelegates.FindOrAdd(GetObjectClass());				\
@@ -124,6 +127,17 @@ struct FFrame;
 	}																																			\
 }
 
+#define REFILL_CUSTOM_ATTRIBUTE_APPLY_DELEGATE(AttributeName, AttributeType)																	\
+if(AttributesAppliedThroughDelegatesKeySet.Contains(Macro_Custom##AttributeName##Key.Key))														\
+{																																				\
+	AttributeType ValueData;																													\
+	if(GetCustom##AttributeName(ValueData))																										\
+	{																																			\
+		constexpr const bool bAddApplyDelegate = true;																							\
+		SetCustom##AttributeName(ValueData, bAddApplyDelegate);																					\
+	}																																			\
+}																																				\
+
 namespace UE::Interchange
 {
 
@@ -133,6 +147,7 @@ namespace UE::Interchange
 	struct FFactoryBaseNodeStaticData : public FBaseNodeStaticData
 	{
 		static INTERCHANGECORE_API const FString& FactoryDependenciesBaseKey();
+		static INTERCHANGECORE_API const FString& AttributesAppliedThroughDelegatesBaseKey();
 		static INTERCHANGECORE_API const FAttributeKey& ReimportStrategyFlagsKey();
 		static INTERCHANGECORE_API const FAttributeKey& SkipNodeImportKey();
 		static INTERCHANGECORE_API const FAttributeKey& ForceNodeReimportKey();
@@ -328,12 +343,22 @@ public:
 
 	INTERCHANGECORE_API void FillAllCustomAttributeFromObject(UObject* Object) const;
 
+	/**
+	 * Adds apply and fill delegates for attributes again that were using custom delegates.
+	 */
+	INTERCHANGECORE_API void RestoreAllCustomAttributeDelegates();
+
 	/**
 	 * Removes all custom attribute delegates registered for the provided class.
 	 * @param Class				The class used to pick which delegates to remove.
 	 */
 	INTERCHANGECORE_API void RemoveCustomAttributesForClass(UClass* Class);
 
+	/**
+	 * Removes all custom attribute delegates that are currently added.
+	 */
+	INTERCHANGECORE_API void RemoveAllCustomAttributeDelegates();
+
 	/**
 	 * Copies all the custom attributes from SourceNode to this node, and
 	 * gets the appropriate values from Object.
@@ -371,6 +396,12 @@ public:
 		return true;
 	}
 
+protected:
+	/**
+	 * Adds apply and fill delegates for attributes again that were using custom delegates.
+	 */
+	INTERCHANGECORE_API virtual void OnRestoreAllCustomAttributeDelegates();
+
 protected:
 	/**
 	 * Those dependencies are used by the Interchange parsing task to make sure the assets are created in the correct order.
@@ -383,6 +414,11 @@ protected:
 
 	TMap<UClass*, TArray<UE::Interchange::FFillAttributeToAsset>> FillCustomAttributeDelegates;
 
+	UE::Interchange::TArrayAttributeHelper<FString> AttributesAppliedThroughDelegates;
+
+	UPROPERTY(Transient)
+	TSet<FString> AttributesAppliedThroughDelegatesKeySet;
+
 private:
 
 	const UE::Interchange::FAttributeKey Macro_CustomSubPathKey = UE::Interchange::FAttributeKey(TEXT("SubPath"));
diff --git a/Engine/Source/Runtime/Interchange/Engine/Private/InterchangeImportReset.cpp b/Engine/Source/Runtime/Interchange/Engine/Private/InterchangeImportReset.cpp
new file mode 100644
--- /dev/null
+++ b/Engine/Source/Runtime/Interchange/Engine/Private/InterchangeImportReset.cpp
@@ -0,0 +1,248 @@
+// Copyright Epic Games, Inc. All Rights Reserved.
+#include "InterchangeImportReset.h"
+
+#include "HAL/IConsoleManager.h"
+#include "InterchangeAssetImportData.h"
+#include "InterchangeEngineLogPrivate.h"
+#include "InterchangeFactoryBase.h"
+#include "InterchangeManager.h"
+#include "InterchangeResultsContainer.h"
+#include "InterchangeSceneImportAsset.h"
+#include "Nodes/InterchangeBaseNodeContainer.h"
+#include "Nodes/InterchangeFactoryBaseNode.h"
+
+DEFINE_LOG_CATEGORY(LogInterchangeReset)
+
+namespace UE::Interchange::InterchangeReset
+{
+	namespace Constants
+	{
+		const FString SceneImportAssetPathKey("InterchangeSceneImportAssetPath");
+	}
+
+	bool GInterchangeResetFilteredNodes = true;
+	FAutoConsoleVariableRef CCvarInterchangeResetFilteredNodes(
+		TEXT("Interchange.FeatureFlags.Reset.UseFilteredNodes"),
+		GInterchangeResetFilteredNodes,
+		TEXT("Whether Reset should be limited to the filtered nodes if any."),
+		ECVF_Default);
+
+	namespace Private
+	{
+		static bool CanExecuteResetPhaseForNode(FInterchangeResetParameters& ResetParameters, const UInterchangeFactoryBaseNode* FactoryNode, const FOnFilterFactoryNodeDelegate& FilterDelegate)
+		{
+			if (!FactoryNode)
+			{
+				return false;
+			}
+
+			if (FilterDelegate.IsBound())
+			{
+				if (!FilterDelegate.Execute(ResetParameters, FactoryNode))
+				{
+					return false;
+				}
+			}
+			else if (!ResetParameters.FilteredNodes.IsEmpty())
+			{
+				if (!ResetParameters.FilteredNodes.Contains(FactoryNode))
+				{
+					return false;
+				}
+			}
+
+			return true;
+		}
+
+		/**
+		 * Helper called before resetting all the imported object properties to the original values
+		 */
+		static void PreResetObjectProperties(FInterchangeResetParameters& ResetParameters, UInterchangeFactoryBaseNode* FactoryNode)
+		{
+			if (!CanExecuteResetPhaseForNode(ResetParameters, FactoryNode, ResetParameters.PreResetDelegates.OnNodeFilter))
+			{
+				return;
+			}
+
+			// This will setup the factory node if is it not already setup.
+			ResetParameters.SetupFactoryNode(FactoryNode);
+
+			if (TObjectPtr<UInterchangeFactoryBase> Factory = ResetParameters.GetFactoryForNode(FactoryNode); IsValid(Factory))
+			{
+				if (UObject* ObjectToReset = ResetParameters.GetReferencedObjectForNode(FactoryNode))
+				{
+					Factory->PreResetObjectProperties(ResetParameters.GetBaseNodeContainer(), FactoryNode, ObjectToReset);
+					ResetParameters.PreResetDelegates.OnNodeProcessed.ExecuteIfBound(ResetParameters, Factory, FactoryNode);
+				}
+			}
+		}
+
+		/**
+		 * Helper used to reset all the properties to the object
+		 */
+		static void ResetObjectProperties(FInterchangeResetParameters& ResetParameters, UInterchangeFactoryBaseNode* FactoryNode)
+		{
+			if (!CanExecuteResetPhaseForNode(ResetParameters, FactoryNode, ResetParameters.ResetDelegates.OnNodeFilter))
+			{
+				return;
+			}
+
+			if (TObjectPtr<UInterchangeFactoryBase> Factory = ResetParameters.GetFactoryForNode(FactoryNode); IsValid(Factory))
+			{
+				if (UObject* ObjectToReset = ResetParameters.GetReferencedObjectForNode(FactoryNode))
+				{
+					Factory->ResetObjectProperties(ResetParameters.GetBaseNodeContainer(), FactoryNode, ObjectToReset);
+					ResetParameters.ResetDelegates.OnNodeProcessed.ExecuteIfBound(ResetParameters, Factory, FactoryNode);
+				}
+			}
+		}
+
+		/**
+		 * Helper called after the object properties are reset
+		 */
+		static void PostResetObjectProperties(FInterchangeResetParameters& ResetParameters, UInterchangeFactoryBaseNode* FactoryNode)
+		{
+			if (!CanExecuteResetPhaseForNode(ResetParameters, FactoryNode, ResetParameters.PostResetDelegates.OnNodeFilter))
+			{
+				return;
+			}
+
+			if (TObjectPtr<UInterchangeFactoryBase> Factory = ResetParameters.GetFactoryForNode(FactoryNode); IsValid(Factory))
+			{
+				if (UObject* ObjectToReset = ResetParameters.GetReferencedObjectForNode(FactoryNode))
+				{
+					Factory->PostResetObjectProperties(ResetParameters.GetBaseNodeContainer(), FactoryNode, ObjectToReset);
+					ResetParameters.PostResetDelegates.OnNodeProcessed.ExecuteIfBound(ResetParameters, Factory, FactoryNode);
+				}
+			}
+		}
+	}
+}
+
+FInterchangeResetParameters::FInterchangeResetParameters(const UInterchangeSceneImportAsset* InSceneImportAsset)
+	:FInterchangeResetParameters(InSceneImportAsset, MakeUnique<FInterchangeResetContextData>())
+{
+}
+
+FInterchangeResetParameters::FInterchangeResetParameters(const UInterchangeSceneImportAsset* InSceneImportAsset, TUniquePtr<FInterchangeResetContextData> InResetContextData)
+	: ResetContextData(MoveTemp(InResetContextData))
+	, SceneImportAsset(InSceneImportAsset)
+	, ResultsContainer(NewObject<UInterchangeResultsContainer>(GetTransientPackage()))
+{
+}
+
+FInterchangeResetParameters::~FInterchangeResetParameters()
+{
+	for (const auto& FactoryNodeDataPair : FactoryNodeDataCache)
+	{
+		if (FactoryNodeDataPair.Value.Factory)
+		{
+			TObjectPtr<UInterchangeFactoryBase> Factory = FactoryNodeDataPair.Value.Factory;
+			Factory->RemoveFromRoot();
+			Factory->ClearFlags(RF_Standalone);
+		}
+	}
+	FactoryNodeDataCache.Empty();
+}
+
+TObjectPtr<UObject> FInterchangeResetParameters::GetReferencedObjectForNode(const UInterchangeFactoryBaseNode* FactoryNode) const
+{
+	if (FactoryNodeDataCache.Contains(FactoryNode))
+	{
+		const FInterchangeResetParameters::FFactoryNodeData& FactoryNodeData = FactoryNodeDataCache[FactoryNode];
+		return Cast<UObject>(FactoryNodeData.ReferencedObjectPath.TryLoad());
+	}
+
+	return TObjectPtr<UObject>(nullptr);
+}
+
+TObjectPtr<UInterchangeFactoryBase> FInterchangeResetParameters::GetFactoryForNode(const UInterchangeFactoryBaseNode* FactoryNode) const
+{
+	if (FactoryNodeDataCache.Contains(FactoryNode))
+	{
+		const FInterchangeResetParameters::FFactoryNodeData& FactoryNodeData = FactoryNodeDataCache[FactoryNode];
+		return FactoryNodeData.Factory;
+	}
+
+	return TObjectPtr<UInterchangeFactoryBase>(nullptr);
+}
+
+void FInterchangeResetParameters::SetupFactoryNode(const UInterchangeFactoryBaseNode* FactoryNode)
+{
+	if (!FactoryNodeDataCache.Contains(FactoryNode))
+	{
+		FInterchangeResetParameters::FFactoryNodeData& FactoryNodeData = FactoryNodeDataCache.FindOrAdd(FactoryNode);
+		UInterchangeManager& InterchangeManager = UInterchangeManager::GetInterchangeManager();
+		if (const UClass* ObjectClass = FactoryNode->GetObjectClass())
+		{
+			if (const UClass* FactoryClass = InterchangeManager.GetRegisteredFactoryClass(ObjectClass))
+			{
+				FactoryNodeData.Factory = NewObject<UInterchangeFactoryBase>(GetTransientPackage(), FactoryClass, NAME_None, RF_Standalone);
+				FactoryNodeData.Factory->AddToRoot();
+				FactoryNodeData.Factory->SetResultsContainer(ResultsContainer);
+			}
+		}
+
+		FSoftObjectPath ReferenceObjectPath;
+		if (FactoryNode->GetCustomReferenceObject(ReferenceObjectPath))
+		{
+			FactoryNodeData.ReferencedObjectPath = ReferenceObjectPath;
+		}
+	}
+}
+
+const UInterchangeBaseNodeContainer* FInterchangeResetParameters::GetBaseNodeContainer() const
+{
+#if WITH_EDITORONLY_DATA
+	if (!SceneImportAsset)
+	{
+		return nullptr;
+	}
+
+	if (!SceneImportAsset->AssetImportData)
+	{
+		return nullptr;
+	}
+
+	return SceneImportAsset->AssetImportData->GetNodeContainer();
+#else
+	return nullptr;
+#endif
+}
+const UInterchangeSceneImportAsset* FInterchangeResetParameters::GetSceneImportAsset() const
+{
+	return SceneImportAsset;
+}
+
+const UInterchangeResultsContainer* FInterchangeResetParameters::GetResultsContainer() const
+{
+	return ResultsContainer;
+}
+
+void FInterchangeReset::ExecuteReset(FInterchangeResetParameters& ResetParameters)
+{
+	TRACE_CPUPROFILER_EVENT_SCOPE(FInterchangeReset::ExecuteReset)
+	if (const UInterchangeBaseNodeContainer* BaseNodeContainer = ResetParameters.GetBaseNodeContainer())
+	{
+		// Pre Reset Phase
+		BaseNodeContainer->IterateNodesOfType<UInterchangeFactoryBaseNode>([&](const FString& NodeKey, UInterchangeFactoryBaseNode* FactoryNode)
+			{
+				UE::Interchange::InterchangeReset::Private::PreResetObjectProperties(ResetParameters, FactoryNode);
+			});
+		ResetParameters.PreResetDelegates.OnCompleted.ExecuteIfBound(ResetParameters);
+
+		// Reset Phase
+		BaseNodeContainer->IterateNodesOfType<UInterchangeFactoryBaseNode>([&](const FString& NodeKey, UInterchangeFactoryBaseNode* FactoryNode)
+			{
+				UE::Interchange::InterchangeReset::Private::ResetObjectProperties(ResetParameters, FactoryNode);
+			});
+		ResetParameters.ResetDelegates.OnCompleted.ExecuteIfBound(ResetParameters);
+
+		// Post Reset Phase
+		BaseNodeContainer->IterateNodesOfType<UInterchangeFactoryBaseNode>([&](const FString& NodeKey, UInterchangeFactoryBaseNode* FactoryNode)
+			{
+				UE::Interchange::InterchangeReset::Private::PostResetObjectProperties(ResetParameters, FactoryNode);
+			});
+		ResetParameters.PostResetDelegates.OnCompleted.ExecuteIfBound(ResetParameters);
+	}
+}
\ No newline at end of file
diff --git a/Engine/Source/Runtime/Interchange/Engine/Private/InterchangeSceneImportAsset.cpp b/Engine/Source/Runtime/Interchange/Engine/Private/InterchangeSceneImportAsset.cpp
--- a/Engine/Source/Runtime/Interchange/Engine/Private/InterchangeSceneImportAsset.cpp
+++ b/Engine/Source/Runtime/Interchange/Engine/Private/InterchangeSceneImportAsset.cpp
@@ -4,6 +4,8 @@
 
 #include "InterchangeAssetImportData.h"
 #include "InterchangeFactoryBase.h"
+#include "InterchangeImportCommon.h"
+#include "InterchangeManager.h"
 
 #include "Engine/AssetUserData.h"
 #include "Engine/Blueprint.h"
@@ -260,24 +262,36 @@ UObject* UInterchangeSceneImportAsset::GetSceneObject(const FString& PackageName
 
 const UInterchangeFactoryBaseNode* UInterchangeSceneImportAsset::GetFactoryNode(const FString& PackageName, const FString& AssetName, const FString& SubPathString) const
 {
+#if WITH_EDITORONLY_DATA
+	const FSoftObjectPath ObjectPath(FName(PackageName), FName(AssetName), SubPathString);
+	return GetFactoryNode(ObjectPath);
+#else
+	return nullptr;
+#endif
+}
+
+const UInterchangeFactoryBaseNode* UInterchangeSceneImportAsset::GetFactoryNode(const FSoftObjectPath& SoftObjectPath) const
+{
 #if WITH_EDITORONLY_DATA
 	if (!ensure(AssetImportData))
 	{
 		return nullptr;
 	}
 
-	const FSoftObjectPath ObjectPath(FName(PackageName), FName(AssetName), SubPathString);
+	const FName ObjectSubPathBaseName = FActorSpawnUtils::GetBaseName(*SoftObjectPath.GetSubPathString());
+	const FString SoftObjectPathSubPathString = SoftObjectPath.GetSubPathString();
 
 	UObject* SceneObject = nullptr;
 	for (const TPair<FSoftObjectPath, FString>& SceneObjectPair : SceneObjects)
 	{
 		const FSoftObjectPath& SceneObjectPath = SceneObjectPair.Key;
-		if (SceneObjectPath.GetLongPackageFName() == ObjectPath.GetLongPackageFName()
-			&& SceneObjectPath.GetAssetFName() == ObjectPath.GetAssetFName())
+		if (SceneObjectPath.GetLongPackageFName() == SoftObjectPath.GetLongPackageFName()
+			&& SceneObjectPath.GetAssetFName() == SoftObjectPath.GetAssetFName())
 		{
 			//World partition actor have a guid that we need to remove before comparison
-			FName SceneObjectSubPathFName = FActorSpawnUtils::GetBaseName(*SceneObjectPath.GetSubPathString());
-			if (SceneObjectSubPathFName == *ObjectPath.GetSubPathString())
+			const FString SceneObjectSubPathString = SceneObjectPath.GetSubPathString();
+			if (SceneObjectSubPathString.Contains(SoftObjectPathSubPathString)
+				&& FActorSpawnUtils::GetBaseName(*SceneObjectSubPathString) == ObjectSubPathBaseName)
 			{
 				return Cast<UInterchangeFactoryBaseNode>(AssetImportData->GetStoredNode(SceneObjectPair.Value));
 			}
@@ -298,4 +312,13 @@ void UInterchangeSceneImportAsset::GetSceneSoftObjectPaths(TArray<FSoftObjectPat
 		SoftObjectPaths.Add(Entry.Key);
 	}
 #endif
+}
+
+bool UInterchangeSceneImportAsset::ContainsObject(const FSoftObjectPath& SoftObjectPath) const
+{
+#if WITH_EDITORONLY_DATA
+	return SceneObjects.Contains(SoftObjectPath);
+#else
+	return false;
+#endif
 }
\ No newline at end of file
diff --git a/Engine/Source/Runtime/Interchange/Engine/Public/InterchangeImportReset.h b/Engine/Source/Runtime/Interchange/Engine/Public/InterchangeImportReset.h
new file mode 100644
--- /dev/null
+++ b/Engine/Source/Runtime/Interchange/Engine/Public/InterchangeImportReset.h
@@ -0,0 +1,95 @@
+// Copyright Epic Games, Inc. All Rights Reserved.
+#pragma once
+
+#include "Containers/Array.h"
+#include "Containers/UnrealString.h"
+#include "CoreTypes.h"
+#include "Delegates/DelegateCombinations.h"
+#include "Templates/UniquePtr.h"
+#include "UObject/ObjectPtr.h"
+#include "UObject/SoftObjectPath.h"
+
+class FAutoConsoleVariableRef;
+
+class UInterchangeBaseNode;
+class UInterchangeBaseNodeContainer;
+class UInterchangeFactoryBase;
+class UInterchangeFactoryBaseNode;
+class UInterchangePipelineBase;
+class UInterchangeTranslatorBase;
+class UInterchangeResultsContainer;
+class UInterchangeSceneImportAsset;
+
+DECLARE_DELEGATE_OneParam(FOnResetPhaseCompletedDelegate, class FInterchangeResetParameters&)
+DECLARE_DELEGATE_RetVal_TwoParams(bool, FOnFilterFactoryNodeDelegate, class FInterchangeResetParameters&, const UInterchangeFactoryBaseNode*)
+DECLARE_DELEGATE_ThreeParams(FOnNodeProcessedDelegate, class FInterchangeResetParameters&, const UInterchangeFactoryBase*, const UInterchangeFactoryBaseNode*)
+
+DECLARE_LOG_CATEGORY_EXTERN(LogInterchangeReset, Display, All)
+
+struct FResetPhaseDelegates
+{
+	FOnFilterFactoryNodeDelegate OnNodeFilter;
+	FOnNodeProcessedDelegate OnNodeProcessed;
+	FOnResetPhaseCompletedDelegate OnCompleted;
+};
+
+/* Used for storing objects that might be needed to be referenced later. */
+struct FInterchangeResetContextData
+{
+	TMap<const UInterchangeFactoryBaseNode*, TArray<TObjectPtr<UObject>>> ObjectsSpawnedDuringReset;
+};
+
+class FInterchangeResetParameters
+{
+	struct FFactoryNodeData
+	{
+		FSoftObjectPath ReferencedObjectPath;
+		TObjectPtr<UInterchangeFactoryBase> Factory;
+	};
+
+public:
+	INTERCHANGEENGINE_API FInterchangeResetParameters(const UInterchangeSceneImportAsset* InSceneImportAsset);
+	INTERCHANGEENGINE_API FInterchangeResetParameters(const UInterchangeSceneImportAsset* InSceneImportAsset, TUniquePtr<FInterchangeResetContextData> InResetContextData);
+	INTERCHANGEENGINE_API ~FInterchangeResetParameters();
+
+	void SetupFactoryNode(const UInterchangeFactoryBaseNode* FactoryNode);
+	TObjectPtr<UObject> GetReferencedObjectForNode(const UInterchangeFactoryBaseNode* FactoryNode) const;
+	TObjectPtr<UInterchangeFactoryBase> GetFactoryForNode(const UInterchangeFactoryBaseNode* FactoryNode) const;
+
+	const UInterchangeBaseNodeContainer* GetBaseNodeContainer() const;
+	const UInterchangeSceneImportAsset* GetSceneImportAsset() const;
+	const UInterchangeResultsContainer* GetResultsContainer() const;
+
+	FResetPhaseDelegates PreResetDelegates;
+	FResetPhaseDelegates ResetDelegates;
+	FResetPhaseDelegates PostResetDelegates;
+
+	TSet<const UInterchangeFactoryBaseNode*> FilteredNodes;
+
+	TUniquePtr<FInterchangeResetContextData> ResetContextData;
+private:
+	TObjectPtr<const UInterchangeSceneImportAsset> SceneImportAsset;
+	TObjectPtr<UInterchangeResultsContainer> ResultsContainer;
+
+	TMap<const UInterchangeFactoryBaseNode*, FFactoryNodeData> FactoryNodeDataCache;
+};
+
+
+class FInterchangeReset
+{
+public:
+	static INTERCHANGEENGINE_API void ExecuteReset(FInterchangeResetParameters& ResetObjectParameters);
+
+};
+
+namespace UE::Interchange::InterchangeReset
+{
+	INTERCHANGEENGINE_API extern FAutoConsoleVariableRef CCvarInterchangeResetFilteredNodes;
+	
+	namespace Constants
+	{
+		INTERCHANGEENGINE_API extern const FString SceneImportAssetPathKey;
+	}
+
+}
+
diff --git a/Engine/Source/Runtime/Interchange/Engine/Public/InterchangeSceneImportAsset.h b/Engine/Source/Runtime/Interchange/Engine/Public/InterchangeSceneImportAsset.h
--- a/Engine/Source/Runtime/Interchange/Engine/Public/InterchangeSceneImportAsset.h
+++ b/Engine/Source/Runtime/Interchange/Engine/Public/InterchangeSceneImportAsset.h
@@ -72,9 +72,13 @@ public:
 	 * @param SubPathString: Optional subobject name
 	 */
 	INTERCHANGEENGINE_API const UInterchangeFactoryBaseNode* GetFactoryNode(const FString& PackageName, const FString& AssetName, const FString& SubPathString = FString()) const;
+	
+	INTERCHANGEENGINE_API const UInterchangeFactoryBaseNode* GetFactoryNode(const FSoftObjectPath& ObjectPath) const;
 
 	INTERCHANGEENGINE_API void GetSceneSoftObjectPaths(TArray<FSoftObjectPath>& SoftObjectPaths);
 
+	INTERCHANGEENGINE_API bool ContainsObject(const FSoftObjectPath& SoftObjectPath) const;
+
 private:
 
 #if WITH_EDITOR
