[BUGFIX] Interchange - Import Dialog - Not all pipeline settings are restored on stack change.

diff --git a/Engine/Plugins/Interchange/Editor/Source/Pipelines/Private/SInterchangePipelineConfigurationDialog.cpp b/Engine/Plugins/Interchange/Editor/Source/Pipelines/Private/SInterchangePipelineConfigurationDialog.cpp
--- a/Engine/Plugins/Interchange/Editor/Source/Pipelines/Private/SInterchangePipelineConfigurationDialog.cpp
+++ b/Engine/Plugins/Interchange/Editor/Source/Pipelines/Private/SInterchangePipelineConfigurationDialog.cpp
@@ -341,6 +341,7 @@ TSharedRef<SBox> SInterchangePipelineConfigurationDialog::SpawnPipelineConfigura
 	{
 		CurrentStackName = FirstStackName;
 	}
+	PreviousStackName = CurrentStackName;
 
 	for(FInterchangeStackInfo& Stack : PipelineStacks)
 	{
@@ -1304,7 +1305,7 @@ void SInterchangePipelineConfigurationDialog::SaveAllPipelineSettings() const
 	{
 		if (PipelineElement->Pipeline)
 		{
-			PipelineElement->Pipeline->SaveSettings(CurrentStackName);
+			PipelineElement->Pipeline->SaveSettings(PreviousStackName);
 		}
 	}
 }
@@ -1336,6 +1337,7 @@ void SInterchangePipelineConfigurationDialog::ClosePipelineConfiguration(const E
 	//Save the settings only if its not a re-import
 	if (!bReimport)
 	{
+		PreviousStackName = CurrentStackName;
 		SaveAllPipelineSettings();
 	}
 
@@ -1470,6 +1472,7 @@ void SInterchangePipelineConfigurationDialog::OnStackSelectionChanged(TSharedPtr
 	}
 
 	//Use the stack select by interchange manager
+	PreviousStackName = CurrentStackName;
 	CurrentStackName = NewStackName;
 
 	constexpr bool bStackSelectionChange = true;
diff --git a/Engine/Plugins/Interchange/Editor/Source/Pipelines/Private/SInterchangePipelineConfigurationDialog.h b/Engine/Plugins/Interchange/Editor/Source/Pipelines/Private/SInterchangePipelineConfigurationDialog.h
--- a/Engine/Plugins/Interchange/Editor/Source/Pipelines/Private/SInterchangePipelineConfigurationDialog.h
+++ b/Engine/Plugins/Interchange/Editor/Source/Pipelines/Private/SInterchangePipelineConfigurationDialog.h
@@ -204,6 +204,7 @@ private:
 	bool bShowSettings = false;
 	bool bShowCards = true;
 
+	FName PreviousStackName = NAME_None;
 	FName CurrentStackName = NAME_None;
 	TObjectPtr<UInterchangePipelineBase> CurrentSelectedPipeline = nullptr;
 };
