Hide Cards View for Interchange Scene Import Asset Reimport

diff --git a/Engine/Plugins/Interchange/Editor/Source/Pipelines/Private/SInterchangePipelineConfigurationDialog.cpp b/Engine/Plugins/Interchange/Editor/Source/Pipelines/Private/SInterchangePipelineConfigurationDialog.cpp
--- a/Engine/Plugins/Interchange/Editor/Source/Pipelines/Private/SInterchangePipelineConfigurationDialog.cpp
+++ b/Engine/Plugins/Interchange/Editor/Source/Pipelines/Private/SInterchangePipelineConfigurationDialog.cpp
@@ -853,7 +853,7 @@ void SInterchangePipelineConfigurationDialog::Construct(const FArguments& InArgs
 
 	//Get the default layout when the user open the import dialog for the first time.
 	bShowEssentials = GInterchangeDefaultShowEssentialsView;
-	bShowCards = !GInterchangeDefaultHideCardsView;
+	bShowCards = !GInterchangeDefaultHideCardsView && !bSceneImport;
 
 	bShowSettings = GInterchangeDefaultShowSettings || !bShowCards;
 
