[Backout] - CL37126820
Original CL Desc
-----------------------------------------------------------------
[BUGFIX] Interchange - Import Dialog - Not all pipeline settings are restored on stack change.

diff --git a/Engine/Plugins/Interchange/Editor/Source/Pipelines/Private/SInterchangePipelineConfigurationDialog.cpp b/Engine/Plugins/Interchange/Editor/Source/Pipelines/Private/SInterchangePipelineConfigurationDialog.cpp
--- a/Engine/Plugins/Interchange/Editor/Source/Pipelines/Private/SInterchangePipelineConfigurationDialog.cpp
+++ b/Engine/Plugins/Interchange/Editor/Source/Pipelines/Private/SInterchangePipelineConfigurationDialog.cpp
@@ -341,7 +341,6 @@ TSharedRef<SBox> SInterchangePipelineConfigurationDialog::SpawnPipelineConfigura
 	{
 		CurrentStackName = FirstStackName;
 	}
-	PreviousStackName = CurrentStackName;
 
 	for(FInterchangeStackInfo& Stack : PipelineStacks)
 	{
@@ -1305,7 +1304,7 @@ void SInterchangePipelineConfigurationDialog::SaveAllPipelineSettings() const
 	{
 		if (PipelineElement->Pipeline)
 		{
-			PipelineElement->Pipeline->SaveSettings(PreviousStackName);
+			PipelineElement->Pipeline->SaveSettings(CurrentStackName);
 		}
 	}
 }
@@ -1337,7 +1336,6 @@ void SInterchangePipelineConfigurationDialog::ClosePipelineConfiguration(const E
 	//Save the settings only if its not a re-import
 	if (!bReimport)
 	{
-		PreviousStackName = CurrentStackName;
 		SaveAllPipelineSettings();
 	}
 
@@ -1371,7 +1369,7 @@ void SInterchangePipelineConfigurationDialog::RefreshStack(bool bStackSelectionC
 	//When doing a reimport we do not want to save the setting because the context have special default
 	//value for some options like: (Import Materials, Import Textures...).
 	//So when doing a reimport switching stack is like doing a reset to default on all pipelines
-	if (!bReimport || !bStackSelectionChange)
+	if (!bReimport && !bStackSelectionChange)
 	{
 		SaveAllPipelineSettings();
 	}
@@ -1471,8 +1469,11 @@ void SInterchangePipelineConfigurationDialog::OnStackSelectionChanged(TSharedPtr
 		return;
 	}
 
-	//Use the stack select by interchange manager
-	PreviousStackName = CurrentStackName;
+	if (!bReimport)
+	{
+		//Use the stack select by interchange manager
+		SaveAllPipelineSettings();
+	}
 	CurrentStackName = NewStackName;
 
 	constexpr bool bStackSelectionChange = true;
diff --git a/Engine/Plugins/Interchange/Editor/Source/Pipelines/Private/SInterchangePipelineConfigurationDialog.h b/Engine/Plugins/Interchange/Editor/Source/Pipelines/Private/SInterchangePipelineConfigurationDialog.h
--- a/Engine/Plugins/Interchange/Editor/Source/Pipelines/Private/SInterchangePipelineConfigurationDialog.h
+++ b/Engine/Plugins/Interchange/Editor/Source/Pipelines/Private/SInterchangePipelineConfigurationDialog.h
@@ -204,7 +204,6 @@ private:
 	bool bShowSettings = false;
 	bool bShowCards = true;
 
-	FName PreviousStackName = NAME_None;
 	FName CurrentStackName = NAME_None;
 	TObjectPtr<UInterchangePipelineBase> CurrentSelectedPipeline = nullptr;
 };
