[Interchange] Fix re-import dialog where some properties could not be edited.

#rnx

diff --git a/Engine/Plugins/Interchange/Editor/Source/Pipelines/Private/SInterchangePipelineConfigurationDialog.cpp b/Engine/Plugins/Interchange/Editor/Source/Pipelines/Private/SInterchangePipelineConfigurationDialog.cpp
--- a/Engine/Plugins/Interchange/Editor/Source/Pipelines/Private/SInterchangePipelineConfigurationDialog.cpp
+++ b/Engine/Plugins/Interchange/Editor/Source/Pipelines/Private/SInterchangePipelineConfigurationDialog.cpp
@@ -1369,7 +1369,7 @@ void SInterchangePipelineConfigurationDialog::RefreshStack(bool bStackSelectionC
 	//When doing a reimport we do not want to save the setting because the context have special default
 	//value for some options like: (Import Materials, Import Textures...).
 	//So when doing a reimport switching stack is like doing a reset to default on all pipelines
-	if (!bReimport && !bStackSelectionChange)
+	if (!bReimport || !bStackSelectionChange)
 	{
 		SaveAllPipelineSettings();
 	}
