[Interchange] Prevent ResetContextMenu from loading worlds whenever one right-clicks on level assets in the ContentBrowser just to see if a reset menu item should be added.

#rnx

diff --git a/Engine/Plugins/Interchange/Editor/Source/InterchangeEditor/Private/InterchangeResetContextMenuExtender.cpp b/Engine/Plugins/Interchange/Editor/Source/InterchangeEditor/Private/InterchangeResetContextMenuExtender.cpp
--- a/Engine/Plugins/Interchange/Editor/Source/InterchangeEditor/Private/InterchangeResetContextMenuExtender.cpp
+++ b/Engine/Plugins/Interchange/Editor/Source/InterchangeEditor/Private/InterchangeResetContextMenuExtender.cpp
@@ -403,7 +403,7 @@ static FDelayedAutoRegisterHelper DelayedAutoRegisterLevelContextMenu(EDelayedRe
 					bool bAddResetSection = false;
 
 					const FAssetData& WorldAsset = Context->SelectedAssets[0];
-					if (Context->SelectedAssets[0].IsInstanceOf<UWorld>())
+					if (Context->SelectedAssets[0].IsInstanceOf<UWorld>() && Context->SelectedAssets[0].IsAssetLoaded())
 					{
 						if (UWorld* World = Cast<UWorld>(Context->SelectedAssets[0].GetAsset()))
 						{
