[Interchange] Fix the UI/UX for interchange import dialog

#rnx

diff --git a/Engine/Plugins/Interchange/Editor/Source/Pipelines/Private/SInterchangeAssetCard.cpp b/Engine/Plugins/Interchange/Editor/Source/Pipelines/Private/SInterchangeAssetCard.cpp
--- a/Engine/Plugins/Interchange/Editor/Source/Pipelines/Private/SInterchangeAssetCard.cpp
+++ b/Engine/Plugins/Interchange/Editor/Source/Pipelines/Private/SInterchangeAssetCard.cpp
@@ -43,8 +43,17 @@ void SInterchangeAssetCard::RefreshCard(UInterchangeBaseNodeContainer* InPreview
 	CardAssetToImportCount = 0;
 	CardAssetDisabledCount = 0;
 	CardTooltip.Reset();
+	FString ImportAssets;
+	FString IgnoreAssets;
+
+	auto FillStringData = [](FString& StringData, const FString& DisplayLabel)
+		{
+			StringData += TEXT("\n\t");
+			StringData += DisplayLabel;
+		};
+
 	//Query the data we need to update the card
-	InPreviewNodeContainer->IterateNodesOfType<UInterchangeFactoryBaseNode>([this](const FString& NodeUid, UInterchangeFactoryBaseNode* FactoryNode)
+	InPreviewNodeContainer->IterateNodesOfType<UInterchangeFactoryBaseNode>([this, &FillStringData, &ImportAssets, &IgnoreAssets](const FString& NodeUid, UInterchangeFactoryBaseNode* FactoryNode)
 		{
 			if (!FactoryNode || !FactoryNode->GetObjectClass() || !FactoryNode->GetObjectClass()->IsChildOf(AssetClass))
 			{
@@ -52,21 +61,28 @@ void SInterchangeAssetCard::RefreshCard(UInterchangeBaseNodeContainer* InPreview
 			}
 
 			CardAssetCount++;
-			if (!CardTooltip.IsEmpty())
-			{
-				CardTooltip += TEXT("\n");
-			}
 			if (FactoryNode->IsEnabled())
 			{
-				CardTooltip += TEXT("Import asset ") + FactoryNode->GetDisplayLabel();
+				FillStringData(ImportAssets, FactoryNode->GetDisplayLabel());
 				CardAssetToImportCount++;
 			}
 			else
 			{
-				CardTooltip += TEXT("Ignore asset ") + FactoryNode->GetDisplayLabel();
+				FillStringData(IgnoreAssets, FactoryNode->GetDisplayLabel());
 				CardAssetDisabledCount++;
 			}
 		});
+
+	if (!ImportAssets.IsEmpty())
+	{
+		CardTooltip += LOCTEXT("CardImportAssetTooltipPrefix", "Import Assets:").ToString();
+		CardTooltip += ImportAssets;
+	}
+	if (!IgnoreAssets.IsEmpty())
+	{
+		CardTooltip += LOCTEXT("CardIgnoreAssetTooltipPrefix", "Ignore Assets:").ToString();
+		CardTooltip += IgnoreAssets;
+	}
 }
 
 bool SInterchangeAssetCard::RefreshHasConflicts(const TArray<FInterchangeConflictInfo>& InConflictInfos)
diff --git a/Engine/Plugins/Interchange/Editor/Source/Pipelines/Private/SInterchangeGraphInspectorWindow.cpp b/Engine/Plugins/Interchange/Editor/Source/Pipelines/Private/SInterchangeGraphInspectorWindow.cpp
--- a/Engine/Plugins/Interchange/Editor/Source/Pipelines/Private/SInterchangeGraphInspectorWindow.cpp
+++ b/Engine/Plugins/Interchange/Editor/Source/Pipelines/Private/SInterchangeGraphInspectorWindow.cpp
@@ -515,8 +515,9 @@ void SInterchangeGraphInspectorWindow::Construct(const FArguments& InArgs)
 			[
 				SNew(SButton)
 				.HAlign(HAlign_Center)
-				.Text(LOCTEXT("InspectorGraphWindow_Close", "Ok"))
+				.Text(LOCTEXT("InspectorGraphWindow_Close", "Dismiss"))
 				.OnClicked(this, &SInterchangeGraphInspectorWindow::OnCloseDialog)
+				.Visibility(bPreview ? EVisibility::Collapsed : EVisibility::Visible)
 			]
 		]
 	];
diff --git a/Engine/Plugins/Interchange/Editor/Source/Pipelines/Private/SInterchangePipelineConfigurationDialog.cpp b/Engine/Plugins/Interchange/Editor/Source/Pipelines/Private/SInterchangePipelineConfigurationDialog.cpp
--- a/Engine/Plugins/Interchange/Editor/Source/Pipelines/Private/SInterchangePipelineConfigurationDialog.cpp
+++ b/Engine/Plugins/Interchange/Editor/Source/Pipelines/Private/SInterchangePipelineConfigurationDialog.cpp
@@ -67,8 +67,7 @@ static FAutoConsoleVariableRef CCvarInterchangeShowConflictWarningOnCardsView(
 	TEXT("Whether the import conflict warnings will be shown on cards view."),
 	ECVF_Default);
 
-constexpr double SplitterSize = 12.0;
-constexpr double AdvancedUIRatio = 1.2;
+constexpr double AdvancedUIRatio = 1.25;
 
 const FName ReimportStackName = TEXT("ReimportPipeline");
 
@@ -234,12 +233,23 @@ void SInterchangePipelineConfigurationDialog::SetEditPipeline(FInterchangePipeli
 
 FReply SInterchangePipelineConfigurationDialog::OnEditTranslatorSettings()
 {
+	if (!TranslatorSettings || !Translator.IsValid())
+	{
+		return FReply::Handled();
+	}
+	UInterchangeTranslatorSettings* TranslatorSettingsCDO = TranslatorSettings->GetClass()->GetDefaultObject<UInterchangeTranslatorSettings>();
+	UInterchangeTranslatorSettings* OriginalTranslatorSettings = DuplicateObject<UInterchangeTranslatorSettings>(TranslatorSettings, GetTransientPackage());
+
 	FPropertyEditorModule& PropertyEditorModule = FModuleManager::GetModuleChecked<FPropertyEditorModule>("PropertyEditor");
 	FDetailsViewArgs DetailsViewArgs;
 	DetailsViewArgs.bAllowSearch = true;
 	DetailsViewArgs.bShowPropertyMatrixButton = false;
 	DetailsViewArgs.bShowSectionSelector = false;
-	DetailsViewArgs.bAllowMultipleTopLevelObjects = true;
+	DetailsViewArgs.bAllowMultipleTopLevelObjects = false;
+	DetailsViewArgs.bShowModifiedPropertiesOption = false;
+	DetailsViewArgs.bShowKeyablePropertiesOption = false;
+	DetailsViewArgs.bShowAnimatedPropertiesOption = false;
+	DetailsViewArgs.bShowHiddenPropertiesWhilePlayingOption = false;
 	DetailsViewArgs.NameAreaSettings = FDetailsViewArgs::HideNameArea;
 	TSharedRef<IDetailsView> TranslatorSettingsDetailsView = PropertyEditorModule.CreateDetailView(DetailsViewArgs);
 
@@ -252,7 +262,7 @@ FReply SInterchangePipelineConfigurationDialog::OnEditTranslatorSettings()
 
 	TSharedRef<SCustomDialog> OptionsDialog =
 		SNew(SCustomDialog)
-		.Title(LOCTEXT("OptionsDialogTitle", "Translator Project Settings Editor"))
+		.Title(LOCTEXT("OptionsDialogTitle", "Translator Settings"))
 		.WindowArguments(SWindow::FArguments()
 			.IsTopmostWindow(true)
 			.MinHeight(300.0f)
@@ -280,10 +290,57 @@ FReply SInterchangePipelineConfigurationDialog::OnEditTranslatorSettings()
 			]
 			
 		]
+		.AutoCloseOnButtonPress(false)
+		.BeforeButtons()
+		[
+			SNew(SHorizontalBox)
+			+ SHorizontalBox::Slot()
+			.Padding(4.f, 0.f)
+			.AutoWidth()
+			[
+				SNew(SButton)
+				.Text(LOCTEXT("DialogButtonResetToDefault", "Reset To Default"))
+				.OnClicked_Lambda([this, &bIsSettingsModified, &TranslatorSettingsCDO, &OriginalTranslatorSettings]()
+					{
+						if (TranslatorSettingsCDO)
+						{
+							//Save the CDO settings
+							TranslatorSettingsCDO->SaveSettings();
+							//load the CDO settings
+							TranslatorSettings->LoadSettings();
+							bIsSettingsModified = true;
+						}
+						return FReply::Handled();
+					})
+			]
+			+ SHorizontalBox::Slot()
+			.FillWidth(1.0f)
+			[
+				SNullWidget::NullWidget
+			]
+		]
 		.Buttons(
 		{
 			SCustomDialog::FButton(LOCTEXT("DialogButtonOk", "Ok"))
+			.SetOnClicked(FSimpleDelegate::CreateLambda([&OptionsDialog]()
+				{
+					OptionsDialog.Get().RequestDestroyWindow();
+				}))
+			.SetPrimary(true),
+			SCustomDialog::FButton(LOCTEXT("DialogButtonCancel", "Cancel"))
+			.SetOnClicked(FSimpleDelegate::CreateLambda([this, &bIsSettingsModified, &OptionsDialog, &OriginalTranslatorSettings]()
+				{
+					if (OriginalTranslatorSettings)
+					{
+						OriginalTranslatorSettings->SaveSettings();
+					}
+					//Revert the settings to what it was before loading the dialog
+					bIsSettingsModified = false;
+					TranslatorSettings->LoadSettings();
+					OptionsDialog.Get().RequestDestroyWindow();
+				}))
 		});
+
 	OptionsDialog->ShowModal();
 
 	if (bIsSettingsModified && Translator.IsValid() && TranslatorSettings)
@@ -316,14 +373,6 @@ TSharedRef<SBox> SInterchangePipelineConfigurationDialog::SpawnPipelineConfigura
 {
 	const ISlateStyle* InterchangeEditorPipelineStyle = FSlateStyleRegistry::FindSlateStyle("InterchangeEditorPipelineStyle");
 
-	const FSlateBrush* ImportSourceBorderBrush = nullptr;
-	const FButtonStyle* ResetSelectPipelineButtonStyle = nullptr;
-	if (InterchangeEditorPipelineStyle)
-	{
-		ImportSourceBorderBrush = InterchangeEditorPipelineStyle->GetBrush("ImportSource.Dropdown.Border");
-		ResetSelectPipelineButtonStyle = &InterchangeEditorPipelineStyle->GetWidgetStyle<FButtonStyle>("ButtonStyle.ResetSelectedPipeline");
-	}
-
 	AvailableStacks.Reset();
 	TSharedPtr<FString> SelectedStack;
 	if (bReimport)
@@ -547,6 +596,10 @@ TSharedRef<SBox> SInterchangePipelineConfigurationDialog::SpawnPipelineConfigura
 	}
 
 	TSharedPtr<SBox> InspectorBox;
+
+	FText FilterPipelineTooltip = LOCTEXT("SInterchangePipelineConfigurationDialog_FilterPipelineOptions_tooltip", "Filter the pipeline options using the source content data.");
+	FText EssentialPipelineTooltip = LOCTEXT("SInterchangePipelineConfigurationDialog_ShowEssentialsOptions_tooltip", "Display only essentials pipeline properties.");
+
 	TSharedRef<SBox> PipelineConfigurationPanelBox = SNew(SBox)
 	[
 		SNew(SVerticalBox)
@@ -573,54 +626,46 @@ TSharedRef<SBox> SInterchangePipelineConfigurationDialog::SpawnPipelineConfigura
 				.AutoWidth()
 				.Padding(4.0f, 0.0f)
 				[
-					SNew(SBorder)
-					.BorderImage(ImportSourceBorderBrush)
-					.Padding(4.0f, 0.0f)
-					.ToolTipText(LOCTEXT("SInterchangePipelineConfigurationDialog_ShowEssentialsOptions_tooltip", "Display only essentials pipeline properties."))
+					SNew(SHorizontalBox)
+					+ SHorizontalBox::Slot()
+					.AutoWidth()
+					.VAlign(VAlign_Center)
+					.Padding(4.f, 0.f)
 					[
-						SNew(SHorizontalBox)
-						+ SHorizontalBox::Slot()
-						.AutoWidth()
-						.VAlign(VAlign_Center)
-						.Padding(4.f, 0.f)
-						[
-							SNew(STextBlock)
-							.Text(LOCTEXT("SInterchangePipelineConfigurationDialog_ShowEssentialsOptions", "Essentials"))
-						]
-						+ SHorizontalBox::Slot()
-						.Padding(4.f, 0.f)
-						[
-							SNew(SCheckBox)
-							.IsChecked(this, &SInterchangePipelineConfigurationDialog::IsShowEssentialsEnabled)
-							.OnCheckStateChanged(this, &SInterchangePipelineConfigurationDialog::OnShowEssentialsChanged)
-						]
+						SNew(STextBlock)
+						.Text(LOCTEXT("SInterchangePipelineConfigurationDialog_ShowEssentialsOptions", "Essentials"))
+						.ToolTipText(EssentialPipelineTooltip)
+					]
+					+ SHorizontalBox::Slot()
+					.Padding(4.f, 0.f)
+					[
+						SNew(SCheckBox)
+						.IsChecked(this, &SInterchangePipelineConfigurationDialog::IsShowEssentialsEnabled)
+						.OnCheckStateChanged(this, &SInterchangePipelineConfigurationDialog::OnShowEssentialsChanged)
+						.ToolTipText(EssentialPipelineTooltip)
 					]
 				]
 				+SHorizontalBox::Slot()
 				.AutoWidth()
 				.Padding(4.0f, 0.0f)
 				[
-					SNew(SBorder)
-					.BorderImage(ImportSourceBorderBrush)
-					.Padding(4.0f, 0.0f)
-					.ToolTipText(LOCTEXT("SInterchangePipelineConfigurationDialog_FilterPipelineOptions_tooltip", "Filter the pipeline options using the source content data."))
+					SNew(SHorizontalBox)
+					+ SHorizontalBox::Slot()
+					.AutoWidth()
+					.VAlign(VAlign_Center)
+					.Padding(4.f, 0.f)
 					[
-						SNew(SHorizontalBox)
-						+ SHorizontalBox::Slot()
-						.AutoWidth()
-						.VAlign(VAlign_Center)
-						.Padding(4.f, 0.f)
-						[
-							SNew(STextBlock)
-							.Text(LOCTEXT("SInterchangePipelineConfigurationDialog_FilterPipelineOptions", "Filter on Contents"))
-						]
-						+ SHorizontalBox::Slot()
-						.Padding(4.f, 0.f)
-						[
-							SNew(SCheckBox)
-							.IsChecked(this, &SInterchangePipelineConfigurationDialog::IsFilteringOptions)
-							.OnCheckStateChanged(this, &SInterchangePipelineConfigurationDialog::OnFilterOptionsChanged)
-						]
+						SNew(STextBlock)
+						.Text(LOCTEXT("SInterchangePipelineConfigurationDialog_FilterPipelineOptions", "Filter on Contents"))
+						.ToolTipText(FilterPipelineTooltip)
+					]
+					+ SHorizontalBox::Slot()
+					.Padding(4.f, 0.f)
+					[
+						SNew(SCheckBox)
+						.IsChecked(this, &SInterchangePipelineConfigurationDialog::IsFilteringOptions)
+						.OnCheckStateChanged(this, &SInterchangePipelineConfigurationDialog::OnFilterOptionsChanged)
+						.ToolTipText(FilterPipelineTooltip)
 					]
 				]
 				+SHorizontalBox::Slot()
@@ -629,21 +674,20 @@ TSharedRef<SBox> SInterchangePipelineConfigurationDialog::SpawnPipelineConfigura
 				.Padding(4.0f, 0.0f, 0.0f, 0.0f)
 				[
 					SNew(SButton)
-					.ButtonStyle(ResetSelectPipelineButtonStyle)
 					.IsEnabled_Lambda([this]()
 						{
 							return PipelinesListView->GetNumItemsSelected() == 1;
 						})
-					.Text(LOCTEXT("SInterchangePipelineConfigurationDialog_ResetToPipelineAsset", "Reset Selected Pipeline"))
+					.Text(LOCTEXT("SInterchangePipelineConfigurationDialog_ResetToPipelineAsset", "Use Pipeline Defaults"))
 					.ToolTipText_Lambda([bReimportClosure = bReimport]()
 						{
 							if (bReimportClosure)
 							{
-								return LOCTEXT("SInterchangePipelineConfigurationDialog_ResetToPipelineAsset_TooltipReimport", "Reset the selected pipeline to the values used the last time this asset was imported.");
+								return LOCTEXT("SInterchangePipelineConfigurationDialog_ResetToPipelineAsset_TooltipReimport", "Reset the selected pipeline to is values used the last time this asset was imported.");
 							}
 							else
 							{
-								return LOCTEXT("SInterchangePipelineConfigurationDialog_ResetToPipelineAsset_Tooltip", "Reset the properties of the selected pipeline.");
+								return LOCTEXT("SInterchangePipelineConfigurationDialog_ResetToPipelineAsset_Tooltip", "Reset the selected pipeline to is default values.");
 							}
 						})
 					.OnClicked(this, &SInterchangePipelineConfigurationDialog::OnResetToDefault)
@@ -677,6 +721,7 @@ TSharedRef<SBox> SInterchangePipelineConfigurationDialog::SpawnPipelineConfigura
 	DetailsViewArgs.bShowModifiedPropertiesOption = false;
 	DetailsViewArgs.bShowKeyablePropertiesOption = false;
 	DetailsViewArgs.bShowAnimatedPropertiesOption = false;
+	DetailsViewArgs.bShowHiddenPropertiesWhilePlayingOption = false;
 	DetailsViewArgs.NameAreaSettings = FDetailsViewArgs::HideNameArea;
 	PipelineConfigurationDetailsView = PropertyEditorModule.CreateDetailView(DetailsViewArgs);
 	InspectorBox->SetContent(PipelineConfigurationDetailsView->AsShared());
@@ -726,40 +771,31 @@ TSharedRef<SBox> SInterchangePipelineConfigurationDialog::SpawnCardsConfiguratio
 			.VAlign(VAlign_Center)
 			[
 				SNew(SCheckBox)
-				.Padding(FMargin(0.f, 4.f))
+				.Padding(FMargin(8.0, 4.0, 0.0, 4.0))
 				.Style(&FAppStyle::Get().GetWidgetStyle<FCheckBoxStyle>("ToggleButtonCheckBox"))
 				.Type(ESlateCheckBoxType::ToggleButton)
 				.IsChecked_Lambda([this]() { return bShowSettings ? ECheckBoxState::Checked : ECheckBoxState::Unchecked; })
 				.OnCheckStateChanged_Lambda([this](ECheckBoxState CheckState)
 					{
 						bShowSettings = CheckState == ECheckBoxState::Checked;
+
 						TSharedPtr<SWindow> DialogWindow = OwnerWindow.Pin();
-							
-						//Resize the window
 						FVector2D ClientSize = DialogWindow->GetClientSizeInScreen();
-						double MaxResize = 1500.0;
 						FWindowSizeLimits SizeLimits = DialogWindow->GetSizeLimits();
-						double MinimumWidth = OriginalMinWindowSize;
-						double NewSize = 0.0;
-						if (bShowSettings)
+						double MinimumSizeX = 0.0;
+						if (bShowCards)
 						{
-							MinimumWidth = OriginalMinWindowSize + (OriginalMinWindowSize * AdvancedUIRatio) + SplitterSize;
-							//Increase the window size in X
-							double SizeRatio = (SplitAdvancedRatio / (1.0f - SplitAdvancedRatio));
-							//Add the splitter to the max size
-							NewSize = ClientSize.X + FMath::Min(MaxResize, (ClientSize.X * SizeRatio)) + SplitterSize;
-							ClientSize.X = FMath::CeilToDouble(NewSize);
+							//Add show cards width
+							MinimumSizeX += OriginalMinWindowSize;
 						}
-						else
+						if (bShowSettings)
 						{
-							MinimumWidth = OriginalMinWindowSize;
-							//Half the window size in X
-							double SizeRatio = (1.0f - SplitAdvancedRatio);
-							NewSize = (ClientSize.X - SplitterSize) * SizeRatio;
-							ClientSize.X = FMath::Floor(NewSize);
+							//Add settings width
+							MinimumSizeX += (OriginalMinWindowSize * AdvancedUIRatio);
 						}
 
-						SizeLimits.SetMinWidth(static_cast<float>(MinimumWidth));
+						//Resize the client with a updated minimum size width
+						SizeLimits.SetMinWidth(static_cast<float>(MinimumSizeX));
 						DialogWindow->SetSizeLimits(SizeLimits);
 						DialogWindow->Resize(ClientSize);
 
@@ -845,7 +881,7 @@ void SInterchangePipelineConfigurationDialog::Construct(const FArguments& InArgs
 		OriginalMinWindowSize = OwnerWindowPinned->GetSizeLimits().GetMinWidth().Get(0.0);
 		if (OriginalMinWindowSize < 1.0)
 		{
-			OriginalMinWindowSize = 350.0;
+			OriginalMinWindowSize = 550.0;
 		}
 		DeltaClientWindowSize = (OwnerWindowPinned->GetSizeInScreen() - OwnerWindowPinned->GetClientSizeInScreen()).X;
 	}
@@ -892,12 +928,6 @@ void SInterchangePipelineConfigurationDialog::Construct(const FArguments& InArgs
 		{
 			//Add settings width
 			MinimumSizeX += (OriginalMinWindowSize * AdvancedUIRatio);
-			
-			//If we have both we have to add the splitter width
-			if (bShowCards)
-			{
-				MinimumSizeX += SplitterSize;
-			}
 		}
 		SizeLimits.SetMinWidth(static_cast<float>(MinimumSizeX));
 		OwnerWindowPinned->SetSizeLimits(SizeLimits);
@@ -1006,7 +1036,7 @@ void SInterchangePipelineConfigurationDialog::Construct(const FArguments& InArgs
 						})
 					.Padding_Lambda([this]()
 						{
-							return (bShowSettings && bShowCards) ? FMargin(0.0f, 0.0f, 20.0f, 0.0f) : FMargin(0.0f, 0.0f, 0.0f, 0.0f);
+							return (bShowSettings && bShowCards) ? FMargin(0.0f, 0.0f, 8.0f, 0.0f) : FMargin(0.0f, 0.0f, 0.0f, 0.0f);
 						})
 					[
 						MainBodyCardsConfiguration
@@ -1025,7 +1055,6 @@ void SInterchangePipelineConfigurationDialog::Construct(const FArguments& InArgs
 						{
 							return bShowCards ? FMargin(8.0f, 0.0f, 0.0f, 0.0f) : FMargin(0.0f, 0.0f, 0.0f, 0.0f);
 						})
-					.Padding(8.0f, 0.0f, 0.0f, 0.0f)
 					.Clipping(EWidgetClipping::ClipToBounds)
 					[
 						MainBodyAdvanced
@@ -1034,6 +1063,13 @@ void SInterchangePipelineConfigurationDialog::Construct(const FArguments& InArgs
 			]
 			+ SVerticalBox::Slot()
 			.AutoHeight()
+			.Padding(4.0f, 8.0, 4.0, 0.0)
+			[
+				SNew(SSeparator)
+				.Orientation(EOrientation::Orient_Horizontal)
+			]
+			+ SVerticalBox::Slot()
+			.AutoHeight()
 			.Padding(0.0f, 8.0f, 0.0f, 0.0f)
 			[
 				SNew(SHorizontalBox)
@@ -1082,6 +1118,17 @@ void SInterchangePipelineConfigurationDialog::Construct(const FArguments& InArgs
 				]
 				+ SHorizontalBox::Slot()
 				.AutoWidth()
+				.Padding(4.f, 0.f) 
+				[
+					SNew(SPrimaryButton)
+					.Icon(this, &SInterchangePipelineConfigurationDialog::GetImportButtonIcon)
+					.Text(LOCTEXT("InspectorGraphWindow_Import", "Import"))
+					.ToolTipText(this, &SInterchangePipelineConfigurationDialog::GetImportButtonTooltip)
+					.IsEnabled(this, &SInterchangePipelineConfigurationDialog::IsImportButtonEnabled)
+					.OnClicked(this, &SInterchangePipelineConfigurationDialog::OnCloseDialog, ECloseEventType::Import)
+				]
+				+ SHorizontalBox::Slot()
+				.AutoWidth()
 				.Padding(4.f, 0.f)
 				[
 					SNew(SButton)
@@ -1096,17 +1143,6 @@ void SInterchangePipelineConfigurationDialog::Construct(const FArguments& InArgs
 				+ SHorizontalBox::Slot()
 				.AutoWidth()
 				.Padding(4.f, 0.f) 
-				[
-					SNew(SPrimaryButton)
-					.Icon(this, &SInterchangePipelineConfigurationDialog::GetImportButtonIcon)
-					.Text(LOCTEXT("InspectorGraphWindow_Import", "Import"))
-					.ToolTipText(this, &SInterchangePipelineConfigurationDialog::GetImportButtonTooltip)
-					.IsEnabled(this, &SInterchangePipelineConfigurationDialog::IsImportButtonEnabled)
-					.OnClicked(this, &SInterchangePipelineConfigurationDialog::OnCloseDialog, ECloseEventType::Import)
-				]
-				+ SHorizontalBox::Slot()
-				.AutoWidth()
-				.Padding(4.f, 0.f) 
 				[
 					SNew(SButton)
 					.HAlign(HAlign_Center)
