Polish import dialog UI
#rnx

diff --git a/Engine/Plugins/Interchange/Editor/Source/Pipelines/Private/SInterchangePipelineConfigurationDialog.cpp b/Engine/Plugins/Interchange/Editor/Source/Pipelines/Private/SInterchangePipelineConfigurationDialog.cpp
--- a/Engine/Plugins/Interchange/Editor/Source/Pipelines/Private/SInterchangePipelineConfigurationDialog.cpp
+++ b/Engine/Plugins/Interchange/Editor/Source/Pipelines/Private/SInterchangePipelineConfigurationDialog.cpp
@@ -265,18 +265,20 @@ FReply SInterchangePipelineConfigurationDialog::OnEditTranslatorSettings()
 		.Title(LOCTEXT("OptionsDialogTitle", "Translator Settings"))
 		.WindowArguments(SWindow::FArguments()
 			.IsTopmostWindow(true)
-			.MinHeight(300.0f)
-			.MinWidth(500)
-			.SizingRule(ESizingRule::UserSized))
+			.MinWidth(400)
+			.MaxHeight(800)
+			.SizingRule(ESizingRule::Autosized))
 		.HAlignContent(HAlign_Fill)
 		.VAlignContent(VAlign_Fill)
 		.HAlignIcon(HAlign_Left)
 		.VAlignIcon(VAlign_Top)
 		.UseScrollBox(true)
+		.RootPadding(FMargin(16.f, 16.f))
+		.ButtonAreaPadding(FMargin(0.f, 16.f, 0.f, 0.f))
 		.Content()
 		[
 			SNew(SBorder)
-			.Padding(FMargin(10.0f, 3.0f))
+			.Padding(FMargin(0.f, 4.f, 0.f, 0.f))
 			.BorderImage(FAppStyle::GetBrush("ToolPanel.DarkGroupBorder"))
 			[
 				SNew(SVerticalBox)
@@ -295,11 +297,11 @@ FReply SInterchangePipelineConfigurationDialog::OnEditTranslatorSettings()
 		[
 			SNew(SHorizontalBox)
 			+ SHorizontalBox::Slot()
-			.Padding(4.f, 0.f)
+			.Padding(0.f, 0.f)
 			.AutoWidth()
 			[
 				SNew(SButton)
-				.Text(LOCTEXT("DialogButtonResetToDefault", "Reset To Default"))
+				.Text(LOCTEXT("DialogButtonResetToDefault", "Reset to Default"))
 				.OnClicked_Lambda([this, &bIsSettingsModified, &TranslatorSettingsCDO, &OriginalTranslatorSettings]()
 					{
 						if (TranslatorSettingsCDO)
@@ -764,7 +766,7 @@ TSharedRef<SBox> SInterchangePipelineConfigurationDialog::SpawnCardsConfiguratio
 			.HAlign(HAlign_Left)
 			[
 				SNew(STextBlock)
-				.Text(LOCTEXT("SInterchangePipelineConfigurationDialog_AssetFoundText", "Assets Found In Source:"))
+				.Text(LOCTEXT("SInterchangePipelineConfigurationDialog_AssetFoundText", "Assets Found in Source:"))
 			]
 			+ SHorizontalBox::Slot()
 			.AutoWidth()
@@ -1063,7 +1065,7 @@ void SInterchangePipelineConfigurationDialog::Construct(const FArguments& InArgs
 			]
 			+ SVerticalBox::Slot()
 			.AutoHeight()
-			.Padding(4.0f, 8.0, 4.0, 0.0)
+			.Padding(0.0f, 8.0, 0.0, 0.0)
 			[
 				SNew(SSeparator)
 				.Orientation(EOrientation::Orient_Horizontal)
@@ -1074,7 +1076,7 @@ void SInterchangePipelineConfigurationDialog::Construct(const FArguments& InArgs
 			[
 				SNew(SHorizontalBox)
 				+ SHorizontalBox::Slot()
-				.Padding(4.f, 0.f)
+				.Padding(0.f, 0.f)
 				.AutoWidth()
 				[
 					IDocumentation::Get()->CreateAnchor(FString("interchange-framework-in-unreal-engine"))
@@ -1118,36 +1120,36 @@ void SInterchangePipelineConfigurationDialog::Construct(const FArguments& InArgs
 				]
 				+ SHorizontalBox::Slot()
 				.AutoWidth()
-				.Padding(4.f, 0.f) 
-				[
-					SNew(SPrimaryButton)
-					.Icon(this, &SInterchangePipelineConfigurationDialog::GetImportButtonIcon)
-					.Text(LOCTEXT("InspectorGraphWindow_Import", "Import"))
-					.ToolTipText(this, &SInterchangePipelineConfigurationDialog::GetImportButtonTooltip)
-					.IsEnabled(this, &SInterchangePipelineConfigurationDialog::IsImportButtonEnabled)
-					.OnClicked(this, &SInterchangePipelineConfigurationDialog::OnCloseDialog, ECloseEventType::Import)
-				]
-				+ SHorizontalBox::Slot()
-				.AutoWidth()
-				.Padding(4.f, 0.f)
-				[
-					SNew(SButton)
-					.Visibility_Lambda([this]()
-						{
-							return bShowSettings ? EVisibility::Visible : EVisibility::Collapsed;
-						})
-					.HAlign(HAlign_Center)
-					.Text(LOCTEXT("InspectorGraphWindow_Preview", "Preview..."))
-					.OnClicked(this, &SInterchangePipelineConfigurationDialog::OnPreviewImport)
-				]
-				+ SHorizontalBox::Slot()
-				.AutoWidth()
-				.Padding(4.f, 0.f) 
 				[
-					SNew(SButton)
-					.HAlign(HAlign_Center)
-					.Text(LOCTEXT("InspectorGraphWindow_Cancel", "Cancel"))
-					.OnClicked(this, &SInterchangePipelineConfigurationDialog::OnCloseDialog, ECloseEventType::Cancel)
+					SNew(SUniformGridPanel)
+					.SlotPadding(FMargin(4.f, 0.f))
+					+SUniformGridPanel::Slot(0,0)
+					[
+						SNew(SPrimaryButton)
+						.Icon(this, &SInterchangePipelineConfigurationDialog::GetImportButtonIcon)
+						.Text(LOCTEXT("InspectorGraphWindow_Import", "Import"))
+						.ToolTipText(this, &SInterchangePipelineConfigurationDialog::GetImportButtonTooltip)
+						.IsEnabled(this, &SInterchangePipelineConfigurationDialog::IsImportButtonEnabled)
+						.OnClicked(this, &SInterchangePipelineConfigurationDialog::OnCloseDialog, ECloseEventType::Import)
+					]
+					+SUniformGridPanel::Slot(1,0)
+					[
+						SNew(SButton)
+						.Visibility_Lambda([this]()
+							{
+								return bShowSettings ? EVisibility::Visible : EVisibility::Collapsed;
+							})
+						.HAlign(HAlign_Center)
+						.Text(LOCTEXT("InspectorGraphWindow_Preview", "Preview..."))
+						.OnClicked(this, &SInterchangePipelineConfigurationDialog::OnPreviewImport)
+					]
+					+SUniformGridPanel::Slot(2,0)
+					[
+						SNew(SButton)
+						.HAlign(HAlign_Center)
+						.Text(LOCTEXT("InspectorGraphWindow_Cancel", "Cancel"))
+						.OnClicked(this, &SInterchangePipelineConfigurationDialog::OnCloseDialog, ECloseEventType::Cancel)
+					]
 				]
 			]
 		]
