Interchange Reset - Automation Tests
- Python tests for testing the Interchange Reset focusing on Mesh Transforms.
- The test will import a gltf scene into level with hierarchy type set to level instance.
- On a successful import, we set the level instance in edit mode to move objects around.
- Then Interchange Reset is invoked, and the transforms are verified to be the same as before, reverting any edits.

#rnx

diff --git a/Engine/Plugins/Interchange/Editor/Source/InterchangeEditor/InterchangeEditor.Build.cs b/Engine/Plugins/Interchange/Editor/Source/InterchangeEditor/InterchangeEditor.Build.cs
--- a/Engine/Plugins/Interchange/Editor/Source/InterchangeEditor/InterchangeEditor.Build.cs
+++ b/Engine/Plugins/Interchange/Editor/Source/InterchangeEditor/InterchangeEditor.Build.cs
@@ -24,6 +24,7 @@ namespace UnrealBuildTool.Rules
 					"AssetDefinition",
 					"AssetTools",
 					"ContentBrowser",
+					"EditorSubsystem",
 					"InterchangeCommonParser",
 					"InterchangeCore",
 					"InterchangeEngine",
diff --git a/Engine/Plugins/Interchange/Editor/Source/InterchangeEditor/Private/AssetDefinition_InterchangeSceneImportAsset.cpp b/Engine/Plugins/Interchange/Editor/Source/InterchangeEditor/Private/AssetDefinition_InterchangeSceneImportAsset.cpp
--- a/Engine/Plugins/Interchange/Editor/Source/InterchangeEditor/Private/AssetDefinition_InterchangeSceneImportAsset.cpp
+++ b/Engine/Plugins/Interchange/Editor/Source/InterchangeEditor/Private/AssetDefinition_InterchangeSceneImportAsset.cpp
@@ -8,6 +8,7 @@
 #include "InterchangeImportReset.h"
 #include "InterchangeLevelFactoryNode.h"
 #include "InterchangeManager.h"
+#include "InterchangeEditorScriptLibrary.h"
 #include "InterchangeSceneImportAsset.h"
 
 #include "ContentBrowserMenuContexts.h"
@@ -44,164 +45,6 @@ UThumbnailInfo* UAssetDefinition_InterchangeSceneImportAsset::LoadThumbnailInfo(
 
 namespace MenuExtension_InterchangeSceneImportAsset
 {
-	ALevelInstance* SpawnLevelInstanceInEditor(FName AssetName)
-	{
-		FWorldContext& EditorWorldContext = GEditor->GetEditorWorldContext();
-		if (UWorld* EditorWorld = EditorWorldContext.World())
-		{
-			FActorSpawnParameters SpawnParameters;
-			SpawnParameters.Name = AssetName;
-			SpawnParameters.NameMode = FActorSpawnParameters::ESpawnActorNameMode::Requested;
-			SpawnParameters.OverrideLevel = EditorWorld->PersistentLevel;
-			SpawnParameters.SpawnCollisionHandlingOverride = ESpawnActorCollisionHandlingMethod::AlwaysSpawn;
-
-			return EditorWorld->SpawnActor<ALevelInstance>(ALevelInstance::StaticClass(), SpawnParameters);
-		}
-
-		return nullptr;
-	}
-
-	void InitializeLevelInstance(ALevelInstance* LevelInstanceActor, UWorld* ReferencedWorld)
-	{
-		if (!LevelInstanceActor->GetRootComponent())
-		{
-			USceneComponent* RootComponent = NewObject<USceneComponent>(LevelInstanceActor, USceneComponent::GetDefaultSceneRootVariableName(), RF_Transactional);
-#if WITH_EDITORONLY_DATA
-			RootComponent->bVisualizeComponent = true;
-#endif
-			LevelInstanceActor->SetRootComponent(RootComponent);
-			LevelInstanceActor->AddInstanceComponent(RootComponent);
-		}
-
-#if WITH_EDITORONLY_DATA
-		UWorld* ParentWorld = LevelInstanceActor->GetWorld();
-
-		// Make sure newly created level asset gets scanned
-		ULevel::ScanLevelAssets(ReferencedWorld->GetPackage()->GetName());
-
-		ParentWorld->PreEditChange(nullptr);
-
-		LevelInstanceActor->SetWorldAsset(ReferencedWorld);
-		LevelInstanceActor->UpdateLevelInstanceFromWorldAsset();
-		LevelInstanceActor->LoadLevelInstance();
-
-		//Reference world must be cleanup since they are not the main world.
-		//This remove all the world managers and prevent GC issue when unloading the main world referencing this world.
-		if (ReferencedWorld->bIsWorldInitialized)
-		{
-			ReferencedWorld->CleanupWorld();
-		}
-
-		ParentWorld->PostEditChange();
-#endif
-	}
-
-
-	void ExecuteResetOneAsset(UInterchangeSceneImportAsset* Asset, const FString& FilePath)
-	{
-		TRACE_CPUPROFILER_EVENT_SCOPE(Interchange_SceneImportAsset_ContextMenuReset)
-		FInterchangeResetParameters ResetParameters(Asset);
-
-		const FSoftObjectPath CurrentSceneImportObjectPath(Asset);
-		const TSubclassOf<UAssetUserData> AssetUserDataClass(UInterchangeLevelAssetUserData::StaticClass());
-
-		ResetParameters.PreResetDelegates.OnNodeProcessed.BindLambda([SceneImportAsset = Asset, &CurrentSceneImportObjectPath, &AssetUserDataClass](FInterchangeResetParameters& ResetParameters, const UInterchangeFactoryBase* Factory, const UInterchangeFactoryBaseNode* FactoryNode)
-		{
-			if (const UInterchangeLevelFactoryNode* LevelFactoryNode = Cast<UInterchangeLevelFactoryNode>(FactoryNode))
-			{
-				bool bShouldCreateLevel = false;
-				LevelFactoryNode->GetCustomShouldCreateLevel(bShouldCreateLevel);
-								
-				if(bShouldCreateLevel)
-				{
-					FSoftObjectPath WorldPath;
-					if (LevelFactoryNode->GetCustomReferenceObject(WorldPath))
-					{
-						if (UWorld* ReferencedWorld = Cast<UWorld>(WorldPath.TryLoad()))
-						{
-							if (AWorldSettings* WorldSettings = ReferencedWorld->GetWorldSettings())
-							{
-								if (UInterchangeLevelAssetUserData* AssetUserData = Cast<UInterchangeLevelAssetUserData>(WorldSettings->GetAssetUserDataOfClass(AssetUserDataClass)))
-								{
-									bool bShouldCreateLevelInstance = false;
-									for (const FSoftObjectPath& SceneImportAssetPath : AssetUserData->SceneImportPaths)
-									{
-										if (SceneImportAssetPath == CurrentSceneImportObjectPath)
-										{
-											bShouldCreateLevelInstance = true;
-											break;
-										}
-									}
-
-									if (!bShouldCreateLevelInstance)
-									{
-										return;
-									}
-
-									TObjectPtr<ALevelInstance> LevelInstance = SpawnLevelInstanceInEditor(FName(ReferencedWorld->GetName()));
-									if (IsValid(LevelInstance))
-									{
-										InitializeLevelInstance(LevelInstance, ReferencedWorld);
-										if (!ResetParameters.ResetContextData->ObjectsSpawnedDuringReset.Contains(FactoryNode))
-										{
-											ResetParameters.ResetContextData->ObjectsSpawnedDuringReset.Emplace(FactoryNode, TArray<TObjectPtr<UObject>>());
-										}
-										
-										ResetParameters.ResetContextData->ObjectsSpawnedDuringReset[FactoryNode].Add(LevelInstance);
-
-										if (LevelInstance->CanEnterEdit())
-										{
-											LevelInstance->EnterEdit();
-										}
-									}
-								}
-							}
-						}
-					}
-				}
-			}
-		});
-
-		ResetParameters.PostResetDelegates.OnNodeProcessed.BindLambda([SceneImportAsset = Asset, &CurrentSceneImportObjectPath, &AssetUserDataClass](FInterchangeResetParameters& ResetParameters, const UInterchangeFactoryBase* Factory, const UInterchangeFactoryBaseNode* FactoryNode)
-			{
-				if (const UInterchangeLevelFactoryNode* LevelFactoryNode = Cast<UInterchangeLevelFactoryNode>(FactoryNode))
-				{
-					bool bShouldCreateLevel = false;
-					LevelFactoryNode->GetCustomShouldCreateLevel(bShouldCreateLevel);
-
-					if (bShouldCreateLevel)
-					{
-						if (ResetParameters.ResetContextData->ObjectsSpawnedDuringReset.Contains(FactoryNode))
-						{
-							for (TObjectPtr<UObject>& SpawnedObject : ResetParameters.ResetContextData->ObjectsSpawnedDuringReset[FactoryNode])
-							{
-								if (ALevelInstance* LevelInstance = Cast<ALevelInstance>(SpawnedObject))
-								{
-									if (LevelInstance->IsEditing())
-									{
-#if WITH_EDITOR
-										// Make sure transformation is up to date after registration as its possible LevelInstance actor can get unregistered when editing properties
-										// through Details panel. In this case the ULevelInstanceComponent might not be able to update the ALevelInstanceEditorInstanceActor transform.
-										LevelInstance->GetLevelInstanceComponent()->UpdateEditorInstanceActor();
-#endif
-										if (LevelInstance->CanExitEdit())
-										{
-											LevelInstance->ExitEdit();
-											// Clear selection in case anything gets selected after applying the changes.
-											GEditor->SelectNone(true, true);
-										}
-									}
-									LevelInstance->Destroy();
-								}
-							}
-						}
-					}
-				}
-			});
-
-		FInterchangeReset::ExecuteReset(ResetParameters);
-		GEditor->RedrawAllViewports();
-	}
 
 	void ExecuteResetScene(const FToolMenuContext& InContext)
 	{
@@ -212,7 +55,7 @@ namespace MenuExtension_InterchangeSceneImportAsset
 		{
 			if (SceneImportAsset)
 			{
-				ExecuteResetOneAsset(SceneImportAsset, SceneImportAsset->AssetImportData->GetFirstFilename());
+				UInterchangeEditorScriptLibrary::ResetSceneImportAsset(SceneImportAsset);
 			}
 		}
 	}
diff --git a/Engine/Plugins/Interchange/Editor/Source/InterchangeEditor/Private/InterchangeEditorScriptLibrary.cpp b/Engine/Plugins/Interchange/Editor/Source/InterchangeEditor/Private/InterchangeEditorScriptLibrary.cpp
new file mode 100644
--- /dev/null
+++ b/Engine/Plugins/Interchange/Editor/Source/InterchangeEditor/Private/InterchangeEditorScriptLibrary.cpp
@@ -0,0 +1,596 @@
+// Copyright Epic Games, Inc. All Rights Reserved.
+#include "InterchangeEditorScriptLibrary.h"
+
+#include "InterchangeAssetImportData.h"
+#include "InterchangeAssetUserData.h"
+#include "InterchangeImportReset.h"
+#include "InterchangeLevelFactoryNode.h"
+#include "InterchangeSceneImportAsset.h"
+
+#include "Editor.h"
+#include "Engine/World.h"
+#include "Engine/Level.h"
+#include "Engine/LevelStreaming.h"
+
+#include "GameFramework/WorldSettings.h"
+#include "GameFramework/Actor.h"
+
+#include "LevelInstance/LevelInstanceActor.h"
+#include "LevelInstance/LevelInstanceComponent.h"
+#include "LevelInstance/LevelInstanceSubsystem.h"
+
+#include "Subsystems/EditorActorSubsystem.h"
+
+
+#define LOCTEXT_NAMESPACE "InterchangeEditorScriptLibrary"
+#pragma optimize("",off)
+namespace UE::InterchangeEditorScriptLibrary::Private
+{
+	namespace InterchangeReset
+	{
+
+		ALevelInstance* SpawnLevelInstanceInEditor(FName AssetName)
+		{
+			FWorldContext& EditorWorldContext = GEditor->GetEditorWorldContext();
+			if (UWorld* EditorWorld = EditorWorldContext.World())
+			{
+				FActorSpawnParameters SpawnParameters;
+				SpawnParameters.Name = AssetName;
+				SpawnParameters.NameMode = FActorSpawnParameters::ESpawnActorNameMode::Requested;
+				SpawnParameters.OverrideLevel = EditorWorld->PersistentLevel;
+				SpawnParameters.SpawnCollisionHandlingOverride = ESpawnActorCollisionHandlingMethod::AlwaysSpawn;
+
+				return EditorWorld->SpawnActor<ALevelInstance>(ALevelInstance::StaticClass(), SpawnParameters);
+			}
+
+			return nullptr;
+		}
+
+		void InitializeLevelInstance(ALevelInstance* LevelInstanceActor, UWorld* ReferencedWorld)
+		{
+			if (!LevelInstanceActor->GetRootComponent())
+			{
+				USceneComponent* RootComponent = NewObject<USceneComponent>(LevelInstanceActor, USceneComponent::GetDefaultSceneRootVariableName(), RF_Transactional);
+#if WITH_EDITORONLY_DATA
+				RootComponent->bVisualizeComponent = true;
+#endif
+				LevelInstanceActor->SetRootComponent(RootComponent);
+				LevelInstanceActor->AddInstanceComponent(RootComponent);
+			}
+
+#if WITH_EDITORONLY_DATA
+			UWorld* ParentWorld = LevelInstanceActor->GetWorld();
+
+			// Make sure newly created level asset gets scanned
+			ULevel::ScanLevelAssets(ReferencedWorld->GetPackage()->GetName());
+
+			ParentWorld->PreEditChange(nullptr);
+
+			LevelInstanceActor->SetWorldAsset(ReferencedWorld);
+			LevelInstanceActor->UpdateLevelInstanceFromWorldAsset();
+			LevelInstanceActor->LoadLevelInstance();
+
+			//Reference world must be cleanup since they are not the main world.
+			//This remove all the world managers and prevent GC issue when unloading the main world referencing this world.
+			if (ReferencedWorld->bIsWorldInitialized)
+			{
+				ReferencedWorld->CleanupWorld();
+			}
+
+			ParentWorld->PostEditChange();
+#endif
+		}
+
+		UInterchangeSceneImportAsset* GetSceneImportAssetFromPath(const FSoftObjectPath& SoftObjectPath)
+		{
+			if (UInterchangeSceneImportAsset* SceneImportAsset = Cast<UInterchangeSceneImportAsset>(SoftObjectPath.TryLoad()))
+			{
+				return SceneImportAsset;
+			}
+
+			return nullptr;
+		}
+
+
+		UInterchangeSceneImportAsset* GetSceneImportAssetForWorld(const UWorld* ReferencedWorld)
+		{
+			if (AWorldSettings* WorldSettings = ReferencedWorld->GetWorldSettings())
+			{
+				if (UInterchangeLevelAssetUserData* AssetUserData = Cast<UInterchangeLevelAssetUserData>(WorldSettings->GetAssetUserDataOfClass(TSubclassOf<UAssetUserData>(UInterchangeLevelAssetUserData::StaticClass()))))
+				{
+					if (!AssetUserData->SceneImportPaths.IsEmpty())
+					{
+						return GetSceneImportAssetFromPath(AssetUserData->SceneImportPaths[0]);
+					}
+				}
+			}
+
+			return nullptr;
+		}
+
+		UInterchangeSceneImportAsset* GetSceneImportAsset(AActor* ImportedActor)
+		{
+			if (const ALevelInstance* LevelInstanceActor = Cast<ALevelInstance>(ImportedActor))
+			{
+				const TSoftObjectPtr<UWorld>& WorldObjectPtr = LevelInstanceActor->GetWorldAsset();
+				if (const UWorld* ReferencedWorld = Cast<UWorld>(WorldObjectPtr.ToSoftObjectPath().TryLoad()))
+				{
+					if (UInterchangeSceneImportAsset* SceneImportAsset = GetSceneImportAssetForWorld(ReferencedWorld))
+					{
+						return SceneImportAsset;
+					}
+				}
+			}
+			else
+			{
+				const TSubclassOf<UAssetUserData> InterchangeAssetUserDataClass(UInterchangeAssetUserData::StaticClass());
+
+				for (UActorComponent* Component : ImportedActor->GetComponents())
+				{
+					if (UInterchangeAssetUserData* AssetUserData = Cast<UInterchangeAssetUserData>(Component->GetAssetUserDataOfClass(InterchangeAssetUserDataClass)))
+					{
+						using namespace UE::Interchange::InterchangeReset;
+						if (const FString* SceneImportAssetPathString = AssetUserData->MetaData.Find(Constants::SceneImportAssetPathKey))
+						{
+							if (const FString* FactoryNodeUid = AssetUserData->MetaData.Find(Constants::FactoryNodeUidPathKey))
+							{
+								if (UInterchangeSceneImportAsset* SceneImportAsset = GetSceneImportAssetFromPath(*SceneImportAssetPathString))
+								{
+									if (SceneImportAsset->AssetImportData->GetStoredFactoryNode(*FactoryNodeUid) != nullptr)
+									{
+										return SceneImportAsset;
+									}
+								}
+							}
+						}
+					}
+				}
+			}
+
+			return nullptr;
+		}
+
+		const UInterchangeFactoryBaseNode* GetFactoryNodeForResetActor(const UInterchangeSceneImportAsset* SceneImportAsset, AActor* ActorToReset)
+		{
+			if (!SceneImportAsset || !SceneImportAsset->AssetImportData || !ActorToReset)
+			{
+				return nullptr;
+			}
+
+			for (UActorComponent* Component : ActorToReset->GetComponents())
+			{
+				if (UInterchangeAssetUserData* AssetUserData = Cast<UInterchangeAssetUserData>(Component->GetAssetUserDataOfClass(UInterchangeAssetUserData::StaticClass())))
+				{
+					using namespace UE::Interchange::InterchangeReset;
+					if (const FString* FactoryNodeUid = AssetUserData->MetaData.Find(Constants::FactoryNodeUidPathKey))
+					{
+						return SceneImportAsset->AssetImportData->GetStoredFactoryNode(*FactoryNodeUid);
+					}
+				}
+			}
+
+			return SceneImportAsset->GetFactoryNode((UObject*)(ActorToReset));
+		}
+
+		void SetWorldAssetForLevelInstanceActor(ALevelInstance* LevelInstanceActor, UWorld* ReferencedWorld)
+		{
+#if WITH_EDITORONLY_DATA
+			UWorld* ParentWorld = LevelInstanceActor->GetWorld();
+
+			// Make sure newly created level asset gets scanned
+			ULevel::ScanLevelAssets(ReferencedWorld->GetPackage()->GetName());
+
+			ParentWorld->PreEditChange(nullptr);
+
+			LevelInstanceActor->SetWorldAsset(ReferencedWorld);
+			LevelInstanceActor->UpdateLevelInstanceFromWorldAsset();
+			LevelInstanceActor->LoadLevelInstance();
+
+			//Reference world must be cleanup since they are not the main world.
+			//This remove all the world managers and prevent GC issue when unloading the main world referencing this world.
+			if (ReferencedWorld->bIsWorldInitialized)
+			{
+				ReferencedWorld->CleanupWorld();
+			}
+
+			ParentWorld->PostEditChange();
+#endif
+		}
+
+		bool ExecuteLevelInstanceReset(const UInterchangeSceneImportAsset* SceneImportAsset, ALevelInstance* LevelInstanceActor)
+		{
+			if (!LevelInstanceActor || !SceneImportAsset)
+			{
+				return false;
+			}
+
+			ILevelInstanceInterface* LevelInstance = Cast<ILevelInstanceInterface>(LevelInstanceActor);
+
+			const TSoftObjectPtr<UWorld>& WorldObjectPtr = LevelInstanceActor->GetWorldAsset();
+
+			if (const UWorld* World = Cast<UWorld>(WorldObjectPtr.ToSoftObjectPath().TryLoad()))
+			{
+				using namespace UE::Interchange::InterchangeReset;
+
+				const bool bWasPrevEditing = LevelInstance->IsEditing();
+				if (!bWasPrevEditing && LevelInstance->CanEnterEdit())
+				{
+					LevelInstance->EnterEdit(LevelInstanceActor);
+				}
+
+				FInterchangeResetParameters ResetParameters(SceneImportAsset);
+
+				// Filter nodes
+				if (CCvarInterchangeResetFilteredNodes->GetBool())
+				{
+					if (ULevel* LoadedLevel = LevelInstance->GetLoadedLevel())
+					{
+						TArray<AActor*> SceneActors = LoadedLevel->Actors;
+						SceneActors.Add(LevelInstanceActor);
+
+						TSet<const UInterchangeFactoryBaseNode*> TempFilteredFactoryNodes;
+						for (AActor* ChildActor : SceneActors)
+						{
+							if (const UInterchangeFactoryBaseNode* FactoryNode = GetFactoryNodeForResetActor(SceneImportAsset, ChildActor))
+							{
+								ResetParameters.AddObjectInstanceToReset(FactoryNode, ChildActor);
+							}
+						}
+					}
+				}
+
+				FInterchangeReset::ExecuteReset(ResetParameters);
+#if WITH_EDITOR
+				// Make sure transformation is up to date after registration as its possible LevelInstance actor can get unregistered when editing properties
+				// through Details panel. In this case the ULevelInstanceComponent might not be able to update the ALevelInstanceEditorInstanceActor transform.
+				LevelInstance->GetLevelInstanceComponent()->UpdateEditorInstanceActor();
+#endif
+				GEditor->RedrawAllViewports();
+
+				if (!bWasPrevEditing && LevelInstance->CanExitEdit())
+				{
+					LevelInstance->ExitEdit();
+				}
+
+				return true;
+			}
+
+			return false;
+		}
+
+		void ExecuteWorldReset(UWorld* World, const UInterchangeSceneImportAsset* SceneImportAsset)
+		{
+			TRACE_CPUPROFILER_EVENT_SCOPE(Interchange_Level_ContextMenuReset)
+				if (!World || !SceneImportAsset)
+				{
+					return;
+				}
+
+			FWorldContext& EditorWorldContext = GEditor->GetEditorWorldContext();
+			if (UWorld* EditorWorld = EditorWorldContext.World())
+			{
+				FActorSpawnParameters SpawnParameters;
+				SpawnParameters.Name = FName(World->GetMapName());
+				SpawnParameters.NameMode = FActorSpawnParameters::ESpawnActorNameMode::Requested;
+				SpawnParameters.OverrideLevel = EditorWorld->PersistentLevel;
+				SpawnParameters.SpawnCollisionHandlingOverride = ESpawnActorCollisionHandlingMethod::AlwaysSpawn;
+
+				if (ALevelInstance* SpawnedActor = EditorWorld->SpawnActor<ALevelInstance>(ALevelInstance::StaticClass(), SpawnParameters))
+				{
+					if (!SpawnedActor->GetRootComponent())
+					{
+						USceneComponent* RootComponent = NewObject<USceneComponent>(SpawnedActor, USceneComponent::GetDefaultSceneRootVariableName(), RF_Transactional);
+#if WITH_EDITORONLY_DATA
+						RootComponent->bVisualizeComponent = true;
+#endif
+						SpawnedActor->SetRootComponent(RootComponent);
+						SpawnedActor->AddInstanceComponent(RootComponent);
+					}
+
+					SetWorldAssetForLevelInstanceActor(SpawnedActor, World);
+					if (ExecuteLevelInstanceReset(SceneImportAsset, SpawnedActor))
+					{
+						GEditor->SelectNone(true, true);
+					}
+
+					SpawnedActor->Destroy();
+				}
+			}
+		}
+	}
+}
+
+void UInterchangeEditorScriptLibrary::ResetLevelAsset(UWorld* World)
+{
+	using namespace UE::InterchangeEditorScriptLibrary::Private::InterchangeReset;
+	if (UInterchangeSceneImportAsset* SceneImportAsset = GetSceneImportAssetForWorld(World))
+	{
+		ExecuteWorldReset(World, SceneImportAsset);
+	}
+}
+
+void UInterchangeEditorScriptLibrary::ResetSceneImportAsset(UInterchangeSceneImportAsset* SceneImportAsset)
+{
+	TRACE_CPUPROFILER_EVENT_SCOPE(Interchange_SceneImportAsset_ContextMenuReset)
+	FInterchangeResetParameters ResetParameters(SceneImportAsset);
+
+	const FSoftObjectPath CurrentSceneImportObjectPath(SceneImportAsset);
+	const TSubclassOf<UAssetUserData> AssetUserDataClass(UInterchangeLevelAssetUserData::StaticClass());
+
+	ResetParameters.PreResetDelegates.OnNodeProcessed.BindLambda([SceneImportAsset, &CurrentSceneImportObjectPath, &AssetUserDataClass](FInterchangeResetParameters& ResetParameters, const UInterchangeFactoryBase* Factory, const UInterchangeFactoryBaseNode* FactoryNode)
+	{
+		if (const UInterchangeLevelFactoryNode* LevelFactoryNode = Cast<UInterchangeLevelFactoryNode>(FactoryNode))
+		{
+			bool bShouldCreateLevel = false;
+			LevelFactoryNode->GetCustomShouldCreateLevel(bShouldCreateLevel);
+
+			if (bShouldCreateLevel)
+			{
+				FSoftObjectPath WorldPath;
+				if (LevelFactoryNode->GetCustomReferenceObject(WorldPath))
+				{
+					if (UWorld* ReferencedWorld = Cast<UWorld>(WorldPath.TryLoad()))
+					{
+						if (AWorldSettings* WorldSettings = ReferencedWorld->GetWorldSettings())
+						{
+							if (UInterchangeLevelAssetUserData* AssetUserData = Cast<UInterchangeLevelAssetUserData>(WorldSettings->GetAssetUserDataOfClass(AssetUserDataClass)))
+							{
+								using namespace UE::InterchangeEditorScriptLibrary::Private::InterchangeReset;
+
+								bool bShouldCreateLevelInstance = false;
+								for (const FSoftObjectPath& SceneImportAssetPath : AssetUserData->SceneImportPaths)
+								{
+									if (SceneImportAssetPath == CurrentSceneImportObjectPath)
+									{
+										bShouldCreateLevelInstance = true;
+										break;
+									}
+								}
+
+								if (!bShouldCreateLevelInstance)
+								{
+									return;
+								}
+
+								TObjectPtr<ALevelInstance> LevelInstance = SpawnLevelInstanceInEditor(FName(ReferencedWorld->GetName()));
+								if (IsValid(LevelInstance))
+								{
+									InitializeLevelInstance(LevelInstance, ReferencedWorld);
+									if (!ResetParameters.ResetContextData->ObjectsSpawnedDuringReset.Contains(FactoryNode))
+									{
+										ResetParameters.ResetContextData->ObjectsSpawnedDuringReset.Emplace(FactoryNode, TArray<TObjectPtr<UObject>>());
+									}
+
+									ResetParameters.ResetContextData->ObjectsSpawnedDuringReset[FactoryNode].Add(LevelInstance);
+
+									if (LevelInstance->CanEnterEdit())
+									{
+										LevelInstance->EnterEdit();
+									}
+								}
+							}
+						}
+					}
+				}
+			}
+		}
+	});
+
+	ResetParameters.PostResetDelegates.OnNodeProcessed.BindLambda([SceneImportAsset, &CurrentSceneImportObjectPath, &AssetUserDataClass](FInterchangeResetParameters& ResetParameters, const UInterchangeFactoryBase* Factory, const UInterchangeFactoryBaseNode* FactoryNode)
+		{
+			if (const UInterchangeLevelFactoryNode* LevelFactoryNode = Cast<UInterchangeLevelFactoryNode>(FactoryNode))
+			{
+				bool bShouldCreateLevel = false;
+				LevelFactoryNode->GetCustomShouldCreateLevel(bShouldCreateLevel);
+
+				if (bShouldCreateLevel)
+				{
+					if (ResetParameters.ResetContextData->ObjectsSpawnedDuringReset.Contains(FactoryNode))
+					{
+						for (TObjectPtr<UObject>& SpawnedObject : ResetParameters.ResetContextData->ObjectsSpawnedDuringReset[FactoryNode])
+						{
+							if (ALevelInstance* LevelInstance = Cast<ALevelInstance>(SpawnedObject))
+							{
+								if (LevelInstance->IsEditing())
+								{
+#if WITH_EDITOR
+									// Make sure transformation is up to date after registration as its possible LevelInstance actor can get unregistered when editing properties
+									// through Details panel. In this case the ULevelInstanceComponent might not be able to update the ALevelInstanceEditorInstanceActor transform.
+									LevelInstance->GetLevelInstanceComponent()->UpdateEditorInstanceActor();
+#endif
+									if (LevelInstance->CanExitEdit())
+									{
+										LevelInstance->ExitEdit();
+										// Clear selection in case anything gets selected after applying the changes.
+										GEditor->SelectNone(true, true);
+									}
+								}
+								LevelInstance->Destroy();
+							}
+						}
+					}
+				}
+			}
+		});
+
+	FInterchangeReset::ExecuteReset(ResetParameters);
+	GEditor->RedrawAllViewports();
+}
+
+void UInterchangeEditorScriptLibrary::ResetActors(TArray<AActor*> Actors)
+{
+	using namespace UE::InterchangeEditorScriptLibrary::Private::InterchangeReset;
+
+	TRACE_CPUPROFILER_EVENT_SCOPE(InterchangeEditorScriptLibrary::ResetActors)
+	TMap<UInterchangeSceneImportAsset*, FInterchangeResetParameters> BatchedReset;
+	TMap<UInterchangeSceneImportAsset*, ALevelInstance*> LevelInstanceResets;
+
+	for (AActor* Actor : Actors)
+	{
+		if (!Actor)
+		{
+			continue;
+		}
+
+		if (UInterchangeSceneImportAsset* SceneImportAsset = GetSceneImportAsset(Actor))
+		{
+			if (ALevelInstance* LevelInstanceActor = Cast<ALevelInstance>(Actor))
+			{
+				if (!LevelInstanceResets.Contains(SceneImportAsset))
+				{
+					// It is enough to reset just one of the level instances as the changes would also be applied to all other when the reset is committed.
+					LevelInstanceResets.Emplace(SceneImportAsset, LevelInstanceActor);
+				}
+			}
+			else
+			{
+				if (!BatchedReset.Contains(SceneImportAsset))
+				{
+					BatchedReset.Emplace(SceneImportAsset, SceneImportAsset);
+				}
+
+				if (UE::Interchange::InterchangeReset::CCvarInterchangeResetFilteredNodes->GetBool())
+				{
+					if (const UInterchangeFactoryBaseNode* FactoryNode = GetFactoryNodeForResetActor(SceneImportAsset, Actor))
+					{
+						BatchedReset[SceneImportAsset].AddObjectInstanceToReset(FactoryNode, Actor);
+					}
+				}
+			}
+		}
+	}
+
+	for (TPair<UInterchangeSceneImportAsset*, FInterchangeResetParameters>& ResetDataPair : BatchedReset)
+	{
+		// As these are all actors in the scene, no need to do anything else to reset them.
+		FInterchangeReset::ExecuteReset(ResetDataPair.Value);
+	}
+
+	for (TPair<UInterchangeSceneImportAsset*, ALevelInstance*>& ResetDataPair : LevelInstanceResets)
+	{
+		ExecuteLevelInstanceReset(ResetDataPair.Key, ResetDataPair.Value);;
+	}
+
+	GEditor->RedrawAllViewports();
+}
+
+bool UInterchangeEditorScriptLibrary::CanResetActor(const AActor* Actor)
+{
+	if (const ALevelInstance* LevelInstanceActor = Cast<ALevelInstance>(Actor))
+	{
+		const TSoftObjectPtr<UWorld>& WorldObjectPtr = LevelInstanceActor->GetWorldAsset();
+		if (const UWorld* ReferencedWorld = Cast<UWorld>(WorldObjectPtr.ToSoftObjectPath().TryLoad()))
+		{
+			if (AWorldSettings* WorldSettings = ReferencedWorld->GetWorldSettings())
+			{
+				return WorldSettings->HasAssetUserDataOfClass(TSubclassOf<UAssetUserData>(UInterchangeLevelAssetUserData::StaticClass()));
+			}
+		}
+	}
+	else
+	{
+		TSubclassOf<UAssetUserData> InterchangeAssetUserDataClass(UInterchangeAssetUserData::StaticClass());
+		for (UActorComponent* Component : Actor->GetComponents())
+		{
+			if (Component->HasAssetUserDataOfClass(InterchangeAssetUserDataClass))
+			{
+				return true;
+			}
+		}
+	}
+
+	return false;
+}
+
+bool UInterchangeEditorScriptLibrary::CanResetWorld(const UWorld* World)
+{
+	if (!World)
+	{
+		return false;
+	}
+
+	if (AWorldSettings* WorldSettings = World->GetWorldSettings())
+	{
+		return WorldSettings->HasAssetUserDataOfClass(TSubclassOf<UAssetUserData>(UInterchangeLevelAssetUserData::StaticClass()));
+	}
+
+	return false;
+}
+
+bool UInterchangeEditorScriptLibrary::LevelInstanceEnterEditMode(ALevelInstance* LevelInstance)
+{
+	if (!LevelInstance)
+	{
+		return false;
+	}
+
+	if (ULevelInstanceSubsystem* LevelInstanceSubsystem = UWorld::GetSubsystem<ULevelInstanceSubsystem>(LevelInstance->GetWorld()))
+	{
+		if (LevelInstanceSubsystem->IsEditingLevelInstance(LevelInstance))
+		{
+			return true;
+		}
+
+		if (LevelInstanceSubsystem->CanEditLevelInstance(LevelInstance))
+		{
+			LevelInstanceSubsystem->EditLevelInstance(LevelInstance);
+			return true;
+		}
+	}
+
+	return false;
+}
+
+bool UInterchangeEditorScriptLibrary::LevelInstanceCommit(ALevelInstance* LevelInstance, bool bDiscardChanges)
+{
+	if (!LevelInstance)
+	{
+		return false;
+	}
+
+	if (ULevelInstanceSubsystem* LevelInstanceSubsystem = UWorld::GetSubsystem<ULevelInstanceSubsystem>(LevelInstance->GetWorld()))
+	{
+		if (LevelInstanceSubsystem->CanCommitLevelInstance(LevelInstance))
+		{
+			LevelInstanceSubsystem->CommitLevelInstance(LevelInstance, bDiscardChanges);
+			return true;
+		}
+	}
+
+	return false;
+}
+
+const TArray<AActor*> UInterchangeEditorScriptLibrary::LevelInstanceGetEditableActors(ALevelInstance* LevelInstance)
+{
+	if (!LevelInstance)
+	{
+		return TArray<AActor*>();
+	}
+
+	if (ULevelInstanceSubsystem* LevelInstanceSubsystem = UWorld::GetSubsystem<ULevelInstanceSubsystem>(LevelInstance->GetWorld()))
+	{
+		if (!LevelInstanceSubsystem->IsEditingLevelInstance(LevelInstance))
+		{
+			return TArray<AActor*>();
+		}
+
+		TSet<AActor*> EditorLevelActorsSet;
+		TSet<AActor*> ReferencedLevelActorsSet;
+
+		if (UEditorActorSubsystem* EditorActorSubsystem = GEditor->GetEditorSubsystem<UEditorActorSubsystem>())
+		{
+			EditorLevelActorsSet.Append(EditorActorSubsystem->GetAllLevelActors());
+		}
+
+		if (ULevel* LoadedLevel = LevelInstance->GetLoadedLevel())
+		{
+			ReferencedLevelActorsSet.Append(LoadedLevel->Actors);
+		}
+
+		return EditorLevelActorsSet.Intersect(ReferencedLevelActorsSet).Array();
+	}
+
+	return TArray<AActor*>();
+}
+#pragma optimize("",on)
+#undef LOCTEXT_NAMESPACE // "InterchangeResetEditorSubsystem"
\ No newline at end of file
diff --git a/Engine/Plugins/Interchange/Editor/Source/InterchangeEditor/Private/InterchangeResetContextMenuExtender.cpp b/Engine/Plugins/Interchange/Editor/Source/InterchangeEditor/Private/InterchangeResetContextMenuExtender.cpp
--- a/Engine/Plugins/Interchange/Editor/Source/InterchangeEditor/Private/InterchangeResetContextMenuExtender.cpp
+++ b/Engine/Plugins/Interchange/Editor/Source/InterchangeEditor/Private/InterchangeResetContextMenuExtender.cpp
@@ -12,7 +12,7 @@
 #include "EngineUtils.h"
 #include "Framework/MultiBox/MultiBoxBuilder.h"
 #include "InterchangeAssetUserData.h"
-#include "InterchangeImportReset.h"
+#include "InterchangeEditorScriptLibrary.h"
 #include "InterchangeLevelFactoryNode.h"
 #include "InterchangeManager.h"
 #include "InterchangeSceneImportAsset.h"
@@ -33,233 +33,15 @@
 
 namespace UE::Interchange::InterchangeReset
 {
-	bool CanResetWorld(const UWorld* World)
-	{
-		if (AWorldSettings* WorldSettings = World->GetWorldSettings())
-		{
-			return WorldSettings->HasAssetUserDataOfClass(TSubclassOf<UAssetUserData>(UInterchangeLevelAssetUserData::StaticClass()));
-		}
-
-		return false;
-	}
-
-	bool CanResetActor(const AActor* Actor)
-	{
-		if (const ALevelInstance* LevelInstanceActor = Cast<ALevelInstance>(Actor))
-		{
-			const TSoftObjectPtr<UWorld>& WorldObjectPtr = LevelInstanceActor->GetWorldAsset();
-			if (const UWorld* ReferencedWorld = Cast<UWorld>(WorldObjectPtr.ToSoftObjectPath().TryLoad()))
-			{
-				return CanResetWorld(ReferencedWorld);
-			}
-		}
-		else
-		{
-			TSubclassOf<UAssetUserData> InterchangeAssetUserDataClass(UInterchangeAssetUserData::StaticClass());
-			for (UActorComponent* Component : Actor->GetComponents())
-			{
-				if (Component->HasAssetUserDataOfClass(InterchangeAssetUserDataClass))
-				{
-					return true;
-				}
-			}
-		}
-
-		return false;
-	}
-
-	UInterchangeSceneImportAsset* GetSceneImportAssetFromPath(const FSoftObjectPath& SoftObjectPath)
-	{
-		if (UInterchangeSceneImportAsset* SceneImportAsset = Cast<UInterchangeSceneImportAsset>(SoftObjectPath.TryLoad()))
-		{
-			return SceneImportAsset;
-		}
-
-		return nullptr;
-	}
-
-
-	UInterchangeSceneImportAsset* GetSceneImportAssetForWorld(const UWorld* ReferencedWorld)
-	{
-		if (AWorldSettings* WorldSettings = ReferencedWorld->GetWorldSettings())
-		{
-			if (UInterchangeLevelAssetUserData* AssetUserData = Cast<UInterchangeLevelAssetUserData>(WorldSettings->GetAssetUserDataOfClass(TSubclassOf<UAssetUserData>(UInterchangeLevelAssetUserData::StaticClass()))))
-			{
-				if (!AssetUserData->SceneImportPaths.IsEmpty())
-				{
-					return GetSceneImportAssetFromPath(AssetUserData->SceneImportPaths[0]);
-				}
-			}
-		}
-
-		return nullptr;
-	}
-
-	UInterchangeSceneImportAsset* GetSceneImportAsset(AActor* ImportedActor)
-	{
-		if (const ALevelInstance* LevelInstanceActor = Cast<ALevelInstance>(ImportedActor))
-		{
-			const TSoftObjectPtr<UWorld>& WorldObjectPtr = LevelInstanceActor->GetWorldAsset();
-			if (const UWorld* ReferencedWorld = Cast<UWorld>(WorldObjectPtr.ToSoftObjectPath().TryLoad()))
-			{
-				if (UInterchangeSceneImportAsset* SceneImportAsset = GetSceneImportAssetForWorld(ReferencedWorld))
-				{
-					return SceneImportAsset;
-				}
-			}
-		}
-		else
-		{
-			const TSubclassOf<UAssetUserData> InterchangeAssetUserDataClass(UInterchangeAssetUserData::StaticClass());
-			
-			for (UActorComponent* Component : ImportedActor->GetComponents())
-			{
-				if (UInterchangeAssetUserData* AssetUserData = Cast<UInterchangeAssetUserData>(Component->GetAssetUserDataOfClass(InterchangeAssetUserDataClass)))
-				{
-					if (const FString* SceneImportAssetPathString = AssetUserData->MetaData.Find(Constants::SceneImportAssetPathKey))
-					{
-						if (UInterchangeSceneImportAsset* SceneImportAsset = GetSceneImportAssetFromPath(*SceneImportAssetPathString);
-							SceneImportAsset && SceneImportAsset->ContainsObject(ImportedActor))
-						{
-							return SceneImportAsset;
-						}
-					}
-				}
-			}
-		}
-
-		return nullptr;
-	}
-	
-	bool ExecuteLevelInstanceReset(const UInterchangeSceneImportAsset* SceneImportAsset, ALevelInstance* LevelInstanceActor)
-	{
-		if (!LevelInstanceActor || !SceneImportAsset)
-		{
-			return false;
-		}
-
-		ILevelInstanceInterface* LevelInstance = Cast<ILevelInstanceInterface>(LevelInstanceActor);
-
-		const TSoftObjectPtr<UWorld>& WorldObjectPtr = LevelInstanceActor->GetWorldAsset();
-
-		if (const UWorld* World = Cast<UWorld>(WorldObjectPtr.ToSoftObjectPath().TryLoad()))
-		{
-			const bool bWasPrevEditing = LevelInstance->IsEditing();
-			if (!bWasPrevEditing && LevelInstance->CanEnterEdit())
-			{
-				LevelInstance->EnterEdit(LevelInstanceActor);
-			}
-
-			FInterchangeResetParameters ResetParameters(SceneImportAsset);
-
-			// Filter nodes
-			if (CCvarInterchangeResetFilteredNodes->GetBool())
-			{
-				if (ULevel* LoadedLevel = LevelInstance->GetLoadedLevel())
-				{
-					TArray<AActor*> SceneActors = LoadedLevel->Actors;
-					SceneActors.Add(LevelInstanceActor);
-
-					TSet<const UInterchangeFactoryBaseNode*> TempFilteredFactoryNodes;
-					for (AActor* Child : SceneActors)
-					{
-						if (const UInterchangeFactoryBaseNode* FactoryNode = SceneImportAsset->GetFactoryNode(Child))
-						{
-							TempFilteredFactoryNodes.Add(FactoryNode);
-						}
-					}
-					ResetParameters.FilteredNodes = MoveTemp(TempFilteredFactoryNodes);
-				}
-			}
-
-			FInterchangeReset::ExecuteReset(ResetParameters);
-#if WITH_EDITOR
-			// Make sure transformation is up to date after registration as its possible LevelInstance actor can get unregistered when editing properties
-			// through Details panel. In this case the ULevelInstanceComponent might not be able to update the ALevelInstanceEditorInstanceActor transform.
-			LevelInstance->GetLevelInstanceComponent()->UpdateEditorInstanceActor();
-#endif
-			GEditor->RedrawAllViewports();
-
-			if (!bWasPrevEditing && LevelInstance->CanExitEdit())
-			{
-				LevelInstance->ExitEdit();
-			}
-
-			return true;
-		}
-		
-		return false;
-	}
-
-	bool ExecuteLevelInstanceReset(TPair<UInterchangeSceneImportAsset*, ALevelInstance*> ResetDataPair)
-	{
-		return ExecuteLevelInstanceReset(ResetDataPair.Key, ResetDataPair.Value);
-	}
-
-	void ResetActors(TArray<AActor*> SelectedActors)
-	{
-		TRACE_CPUPROFILER_EVENT_SCOPE(Interchange_Actors_ContextMenuReset)
-		TMap<UInterchangeSceneImportAsset*, FInterchangeResetParameters> BatchedReset;
-		TMap<UInterchangeSceneImportAsset*, ALevelInstance*> LevelInstanceResets;
-
-		for (AActor* Actor : SelectedActors)
-		{
-			if (!Actor)
-			{
-				continue;
-			}
-
-			if (UInterchangeSceneImportAsset* SceneImportAsset = GetSceneImportAsset(Actor))
-			{
-				if (ALevelInstance* LevelInstanceActor = Cast<ALevelInstance>(Actor))
-				{
-					if (!LevelInstanceResets.Contains(SceneImportAsset))
-					{
-						// It is enough to reset just one of the level instances as the changes would also be applied to all other when the reset is committed.
-						LevelInstanceResets.Emplace(SceneImportAsset, LevelInstanceActor);
-					}
-				}
-				else
-				{
-					if (!BatchedReset.Contains(SceneImportAsset))
-					{
-						BatchedReset.Emplace(SceneImportAsset, SceneImportAsset);
-					}
-
-					if (CCvarInterchangeResetFilteredNodes->GetBool())
-					{
-						if (const UInterchangeFactoryBaseNode* FactoryNode = SceneImportAsset->GetFactoryNode(Actor))
-						{
-							BatchedReset[SceneImportAsset].FilteredNodes.Add(FactoryNode);
-						}
-					}
-				}
-			}
-		}
-
-		for (TPair<UInterchangeSceneImportAsset*, FInterchangeResetParameters>& ResetDataPair : BatchedReset)
-		{
-			// As these are all actors in the scene, no need to do anything else to reset them.
-			FInterchangeReset::ExecuteReset(ResetDataPair.Value);
-		}
-
-		for (TPair<UInterchangeSceneImportAsset*, ALevelInstance*>& ResetDataPair : LevelInstanceResets)
-		{
-			ExecuteLevelInstanceReset(ResetDataPair);
-		}
-
-		GEditor->RedrawAllViewports();
-	}
-
 	TSharedRef<FExtender> OnExtendLevelEditorActorSelectionMenu(const TSharedRef<FUICommandList> CommandList, TArray<AActor*> SelectedActors)
 	{
 		TSharedRef<FExtender> Extender = MakeShared<FExtender>();
 
 		bool bShouldExtendActorActions = false;
 
-		for (AActor* Actor : SelectedActors)
+		for (const AActor* Actor : SelectedActors)
 		{
-			if (CanResetActor(Actor))
+			if (UInterchangeEditorScriptLibrary::CanResetActor(Actor))
 			{
 				bShouldExtendActorActions = true;
 				break;
@@ -277,7 +59,7 @@ namespace UE::Interchange::InterchangeReset
 						NSLOCTEXT("InterchangeActions", "ObjectContext_ResetInterchange", "Reset Properties"),
 						NSLOCTEXT("InterchangeActions", "ObjectContext_ResetInterchangeTooltip", "Resets overridden values with the values from Interchange Import"),
 						FSlateIcon(FAppStyle::GetAppStyleSetName(), "SourceControl.Actions.Refresh"),
-						FUIAction(FExecuteAction::CreateStatic(&ResetActors, SelectedActors), FCanExecuteAction()));
+						FUIAction(FExecuteAction::CreateStatic(&UInterchangeEditorScriptLibrary::ResetActors, SelectedActors), FCanExecuteAction()));
 					MenuBuilder.EndSection();
 				}));
 		}
@@ -287,81 +69,16 @@ namespace UE::Interchange::InterchangeReset
 
 	namespace LevelContextMenuHelpers
 	{
-		void UpdateLevelInstanceActorWorld(ALevelInstance* LevelInstanceActor, UWorld* ReferencedWorld)
-		{
-#if WITH_EDITORONLY_DATA
-			UWorld* ParentWorld = LevelInstanceActor->GetWorld();
-
-			// Make sure newly created level asset gets scanned
-			ULevel::ScanLevelAssets(ReferencedWorld->GetPackage()->GetName());
-
-			ParentWorld->PreEditChange(nullptr);
-
-			LevelInstanceActor->SetWorldAsset(ReferencedWorld);
-			LevelInstanceActor->UpdateLevelInstanceFromWorldAsset();
-			LevelInstanceActor->LoadLevelInstance();
-
-			//Reference world must be cleanup since they are not the main world.
-			//This remove all the world managers and prevent GC issue when unloading the main world referencing this world.
-			if (ReferencedWorld->bIsWorldInitialized)
-			{
-				ReferencedWorld->CleanupWorld();
-			}
-
-			ParentWorld->PostEditChange();
-#endif
-		}
-
 		bool CanExecuteResetLevel(const FToolMenuContext& InContext)
 		{
 			const UContentBrowserAssetContextMenuContext* Context = UContentBrowserAssetContextMenuContext::FindContextWithAssets(InContext);
 			TArray<UWorld*> Worlds = Context->LoadSelectedObjects<UWorld>();
-			return Algo::AnyOf(Worlds, [](UWorld* World)
+			return Algo::AnyOf(Worlds, [](const UWorld* World)
 				{
-					return World && CanResetWorld(World);
+					return UInterchangeEditorScriptLibrary::CanResetWorld(World);
 				});
 		}
 
-		void ExecuteResetOneLevel(UWorld* World, const UInterchangeSceneImportAsset* SceneImportAsset)
-		{
-			TRACE_CPUPROFILER_EVENT_SCOPE(Interchange_Level_ContextMenuReset)
-			if (!World || !SceneImportAsset)
-			{
-				return;
-			}
-
-			FWorldContext& EditorWorldContext = GEditor->GetEditorWorldContext();
-			if (UWorld* EditorWorld = EditorWorldContext.World())
-			{
-				FActorSpawnParameters SpawnParameters;
-				SpawnParameters.Name = FName(World->GetMapName());
-				SpawnParameters.NameMode = FActorSpawnParameters::ESpawnActorNameMode::Requested;
-				SpawnParameters.OverrideLevel = EditorWorld->PersistentLevel;
-				SpawnParameters.SpawnCollisionHandlingOverride = ESpawnActorCollisionHandlingMethod::AlwaysSpawn;
-
-				if (ALevelInstance* SpawnedActor = EditorWorld->SpawnActor<ALevelInstance>(ALevelInstance::StaticClass(), SpawnParameters))
-				{
-					if (!SpawnedActor->GetRootComponent())
-					{
-						USceneComponent* RootComponent = NewObject<USceneComponent>(SpawnedActor, USceneComponent::GetDefaultSceneRootVariableName(), RF_Transactional);
-#if WITH_EDITORONLY_DATA
-						RootComponent->bVisualizeComponent = true;
-#endif
-						SpawnedActor->SetRootComponent(RootComponent);
-						SpawnedActor->AddInstanceComponent(RootComponent);
-					}
-
-					UpdateLevelInstanceActorWorld(SpawnedActor, World);
-					if (ExecuteLevelInstanceReset(SceneImportAsset, SpawnedActor))
-					{
-						GEditor->SelectNone(true, true);
-					}
-
-					SpawnedActor->Destroy();
-				}
-			}
-		}
-
 		void ExecuteResetLevel(const FToolMenuContext& InContext)
 		{
 			const UContentBrowserAssetContextMenuContext* Context = UContentBrowserAssetContextMenuContext::FindContextWithAssets(InContext);
@@ -369,10 +86,7 @@ namespace UE::Interchange::InterchangeReset
 
 			for (UWorld* World : Worlds)
 			{
-				if (UInterchangeSceneImportAsset* SceneImportAsset = GetSceneImportAssetForWorld(World))
-				{
-					ExecuteResetOneLevel(World, SceneImportAsset);
-				}
+				UInterchangeEditorScriptLibrary::ResetLevelAsset(World);	
 			}
 		}
 	}
@@ -407,12 +121,9 @@ static FDelayedAutoRegisterHelper DelayedAutoRegisterLevelContextMenu(EDelayedRe
 					{
 						if (UWorld* World = Cast<UWorld>(Context->SelectedAssets[0].GetAsset()))
 						{
-							if (AWorldSettings* WorldSettings = World->GetWorldSettings())
+							if (UInterchangeEditorScriptLibrary::CanResetWorld(World))
 							{
-								if (WorldSettings->HasAssetUserDataOfClass(TSubclassOf<UAssetUserData>(UInterchangeLevelAssetUserData::StaticClass())))
-								{
-									bAddResetSection = true;
-								}
+								bAddResetSection = true;
 							}
 						}
 					}
diff --git a/Engine/Plugins/Interchange/Editor/Source/InterchangeEditor/Public/InterchangeEditorScriptLibrary.h b/Engine/Plugins/Interchange/Editor/Source/InterchangeEditor/Public/InterchangeEditorScriptLibrary.h
new file mode 100644
--- /dev/null
+++ b/Engine/Plugins/Interchange/Editor/Source/InterchangeEditor/Public/InterchangeEditorScriptLibrary.h
@@ -0,0 +1,84 @@
+// Copyright Epic Games, Inc. All Rights Reserved.
+#pragma once
+#include "Kismet/BlueprintFunctionLibrary.h"
+#include "Containers/Array.h"
+
+#include "InterchangeEditorScriptLibrary.generated.h"
+
+class AActor;
+class ALevelInstance;
+
+class UWorld;
+class UInterchangeSceneImportAsset;
+
+UCLASS()
+class INTERCHANGEEDITOR_API UInterchangeEditorScriptLibrary : public UBlueprintFunctionLibrary
+{
+	GENERATED_BODY()
+public:
+	/**
+	 * Performs Interchange Reset on a Level Asset.
+	 * @param World is the level asset to reset.
+	 */
+	UFUNCTION(BlueprintCallable, Category="Interchange Utilities | Reset")
+	static void ResetLevelAsset(UWorld* World);
+
+	/**
+	 * Performs Interchange Reset on an Interchange Scene Import Asset.
+	 * Resets all the actors added to the level and assets imported.
+	 * @param SceneImportAsset to reset
+	 */
+	UFUNCTION(BlueprintCallable, Category = "Interchange Utilities | Reset")
+	static void ResetSceneImportAsset(UInterchangeSceneImportAsset* SceneImportAsset);
+
+	/**
+	 * Performs Interchange Reset on Actors.
+	 * Resets all qualifying actors. Does nothing to actors that cannot be reset.
+	 * @param Actors array containing actors to reset
+	 */
+	UFUNCTION(BlueprintCallable, Category = "Interchange Utilities | Reset")
+	static void ResetActors(TArray<AActor*> Actors);
+
+	/**
+	 * Checks if an actor can be reset.
+	 * @param Actor to check
+	 * @return can actor be reset.
+	 */
+	UFUNCTION(BlueprintCallable, Category = "Interchange Utilities | Reset")
+	static bool CanResetActor(const AActor* Actor);
+
+	/**
+	 * Checks if an world can be reset.
+	 * @param World to check
+	 * @return can world be reset.
+	 */
+	UFUNCTION(BlueprintCallable, Category = "Interchange Utilities | Reset")
+	static bool CanResetWorld(const UWorld* World);
+
+	/**
+	 * Make Level Instance Actor editable.
+	 * @param LevelInstance actor to edit
+	 * @return whether putting level instance in edit mode was successful.
+	 */
+	UFUNCTION(BlueprintCallable, Category = "Interchange Utilities | Reset")
+	static bool LevelInstanceEnterEditMode(ALevelInstance* LevelInstance);
+
+	/**
+	 * Apply/Discard the changes to Level Instance Actor.
+	 * @param LevelInstance actor to commit
+	 * @param bDiscardChanges should changes be applied or discarded.
+	 * @return whether committing level instance was successful.
+	 */
+	UFUNCTION(BlueprintCallable, Category = "Interchange Utilities | Reset")
+	static bool LevelInstanceCommit(ALevelInstance* LevelInstance, bool bDiscardChanges);
+
+	/**
+	 * Returns array of actors that are editable in the editor when the level instance is put in edit mode.
+	 * NOTE: This will return a non-empty array if the LevelInstance is put in the edit mode.
+	 * 
+	 * @param LevelInstance to get editable actors from
+	 * @return Array of editable actors.
+	 */
+	UFUNCTION(BlueprintCallable, Category = "Interchange Utilities | Reset")
+	static const TArray<AActor*> LevelInstanceGetEditableActors(ALevelInstance* LevelInstance);
+};
diff --git a/Engine/Source/Runtime/Interchange/Engine/Private/InterchangeImportReset.cpp b/Engine/Source/Runtime/Interchange/Engine/Private/InterchangeImportReset.cpp
--- a/Engine/Source/Runtime/Interchange/Engine/Private/InterchangeImportReset.cpp
+++ b/Engine/Source/Runtime/Interchange/Engine/Private/InterchangeImportReset.cpp
@@ -12,12 +12,12 @@
 #include "Nodes/InterchangeFactoryBaseNode.h"
 
 DEFINE_LOG_CATEGORY(LogInterchangeReset)
-
 namespace UE::Interchange::InterchangeReset
 {
 	namespace Constants
 	{
 		const FString SceneImportAssetPathKey("InterchangeSceneImportAssetPath");
+		const FString FactoryNodeUidPathKey("InterchangeFactoryNodeUid");
 	}
 
 	bool GInterchangeResetFilteredNodes = true;
@@ -69,11 +69,12 @@ namespace UE::Interchange::InterchangeReset
 
 			if (TObjectPtr<UInterchangeFactoryBase> Factory = ResetParameters.GetFactoryForNode(FactoryNode); IsValid(Factory))
 			{
-				if (UObject* ObjectToReset = ResetParameters.GetReferencedObjectForNode(FactoryNode))
+				TArray<TObjectPtr<UObject>> ObjectInstances = ResetParameters.GetObjectInstancesForFactoryNode(FactoryNode);
+				for (TObjectPtr<UObject>& ObjectInstance : ObjectInstances)
 				{
-					Factory->PreResetObjectProperties(ResetParameters.GetBaseNodeContainer(), FactoryNode, ObjectToReset);
-					ResetParameters.PreResetDelegates.OnNodeProcessed.ExecuteIfBound(ResetParameters, Factory, FactoryNode);
+					Factory->PreResetObjectProperties(ResetParameters.GetBaseNodeContainer(), FactoryNode, ObjectInstance);
 				}
+				ResetParameters.PreResetDelegates.OnNodeProcessed.ExecuteIfBound(ResetParameters, Factory, FactoryNode);
 			}
 		}
 
@@ -89,11 +90,12 @@ namespace UE::Interchange::InterchangeReset
 
 			if (TObjectPtr<UInterchangeFactoryBase> Factory = ResetParameters.GetFactoryForNode(FactoryNode); IsValid(Factory))
 			{
-				if (UObject* ObjectToReset = ResetParameters.GetReferencedObjectForNode(FactoryNode))
+				TArray<TObjectPtr<UObject>> ObjectInstances = ResetParameters.GetObjectInstancesForFactoryNode(FactoryNode);
+				for (TObjectPtr<UObject>& ObjectInstance : ObjectInstances)
 				{
-					Factory->ResetObjectProperties(ResetParameters.GetBaseNodeContainer(), FactoryNode, ObjectToReset);
-					ResetParameters.ResetDelegates.OnNodeProcessed.ExecuteIfBound(ResetParameters, Factory, FactoryNode);
+					Factory->ResetObjectProperties(ResetParameters.GetBaseNodeContainer(), FactoryNode, ObjectInstance);
 				}
+				ResetParameters.ResetDelegates.OnNodeProcessed.ExecuteIfBound(ResetParameters, Factory, FactoryNode);
 			}
 		}
 
@@ -109,11 +111,12 @@ namespace UE::Interchange::InterchangeReset
 
 			if (TObjectPtr<UInterchangeFactoryBase> Factory = ResetParameters.GetFactoryForNode(FactoryNode); IsValid(Factory))
 			{
-				if (UObject* ObjectToReset = ResetParameters.GetReferencedObjectForNode(FactoryNode))
+				TArray<TObjectPtr<UObject>> ObjectInstances = ResetParameters.GetObjectInstancesForFactoryNode(FactoryNode);
+				for(TObjectPtr<UObject>& ObjectInstance : ObjectInstances)
 				{
-					Factory->PostResetObjectProperties(ResetParameters.GetBaseNodeContainer(), FactoryNode, ObjectToReset);
-					ResetParameters.PostResetDelegates.OnNodeProcessed.ExecuteIfBound(ResetParameters, Factory, FactoryNode);
+					Factory->PostResetObjectProperties(ResetParameters.GetBaseNodeContainer(), FactoryNode, ObjectInstance);
 				}
+				ResetParameters.PostResetDelegates.OnNodeProcessed.ExecuteIfBound(ResetParameters, Factory, FactoryNode);
 			}
 		}
 	}
@@ -145,33 +148,50 @@ FInterchangeResetParameters::~FInterchangeResetParameters()
 	FactoryNodeDataCache.Empty();
 }
 
-TObjectPtr<UObject> FInterchangeResetParameters::GetReferencedObjectForNode(const UInterchangeFactoryBaseNode* FactoryNode) const
+TObjectPtr<UInterchangeFactoryBase> FInterchangeResetParameters::GetFactoryForNode(const UInterchangeFactoryBaseNode* FactoryNode) const
 {
 	if (FactoryNodeDataCache.Contains(FactoryNode))
 	{
 		const FInterchangeResetParameters::FFactoryNodeData& FactoryNodeData = FactoryNodeDataCache[FactoryNode];
-		return Cast<UObject>(FactoryNodeData.ReferencedObjectPath.TryLoad());
+		return FactoryNodeData.Factory;
 	}
 
-	return TObjectPtr<UObject>(nullptr);
+	return TObjectPtr<UInterchangeFactoryBase>(nullptr);
 }
 
-TObjectPtr<UInterchangeFactoryBase> FInterchangeResetParameters::GetFactoryForNode(const UInterchangeFactoryBaseNode* FactoryNode) const
+TArray<TObjectPtr<UObject>> FInterchangeResetParameters::GetObjectInstancesForFactoryNode(const UInterchangeFactoryBaseNode* FactoryNode)
 {
 	if (FactoryNodeDataCache.Contains(FactoryNode))
 	{
 		const FInterchangeResetParameters::FFactoryNodeData& FactoryNodeData = FactoryNodeDataCache[FactoryNode];
-		return FactoryNodeData.Factory;
+		if (FactoryNodeData.ObjectsToReset.Num())
+		{
+			return FactoryNodeData.ObjectsToReset;
+		}
+		else
+		{
+			return TArray<TObjectPtr<UObject>>{ Cast<UObject>(FactoryNodeData.ReferencedObjectPath.TryLoad()) };
+		}
 	}
 
-	return TObjectPtr<UInterchangeFactoryBase>(nullptr);
+	return TArray<TObjectPtr<UObject>>();
+}
+
+void FInterchangeResetParameters::AddObjectInstanceToReset(const UInterchangeFactoryBaseNode* FactoryNode, UObject* ObjectToReset)
+{
+	FInterchangeResetParameters::FFactoryNodeData& FactoryNodeData = FactoryNodeDataCache.FindOrAdd(FactoryNode);
+	if (!FilteredNodes.Contains(FactoryNode))
+	{
+		FilteredNodes.Add(FactoryNode);
+	}
+	FactoryNodeDataCache[FactoryNode].ObjectsToReset.Add(ObjectToReset);
 }
 
 void FInterchangeResetParameters::SetupFactoryNode(const UInterchangeFactoryBaseNode* FactoryNode)
 {
-	if (!FactoryNodeDataCache.Contains(FactoryNode))
+	FInterchangeResetParameters::FFactoryNodeData& FactoryNodeData = FactoryNodeDataCache.FindOrAdd(FactoryNode);
+	if (!FactoryNodeData.Factory)
 	{
-		FInterchangeResetParameters::FFactoryNodeData& FactoryNodeData = FactoryNodeDataCache.FindOrAdd(FactoryNode);
 		UInterchangeManager& InterchangeManager = UInterchangeManager::GetInterchangeManager();
 		if (const UClass* ObjectClass = FactoryNode->GetObjectClass())
 		{
diff --git a/Engine/Source/Runtime/Interchange/Engine/Public/InterchangeImportReset.h b/Engine/Source/Runtime/Interchange/Engine/Public/InterchangeImportReset.h
--- a/Engine/Source/Runtime/Interchange/Engine/Public/InterchangeImportReset.h
+++ b/Engine/Source/Runtime/Interchange/Engine/Public/InterchangeImportReset.h
@@ -45,6 +45,7 @@ class FInterchangeResetParameters
 	{
 		FSoftObjectPath ReferencedObjectPath;
 		TObjectPtr<UInterchangeFactoryBase> Factory;
+		TArray<TObjectPtr<UObject>> ObjectsToReset;
 	};
 
 public:
@@ -52,9 +53,12 @@ public:
 	INTERCHANGEENGINE_API FInterchangeResetParameters(const UInterchangeSceneImportAsset* InSceneImportAsset, TUniquePtr<FInterchangeResetContextData> InResetContextData);
 	INTERCHANGEENGINE_API ~FInterchangeResetParameters();
 
+	INTERCHANGEENGINE_API void AddObjectInstanceToReset(const UInterchangeFactoryBaseNode* FactoryNode, UObject* ObjectToReset);
+
 	void SetupFactoryNode(const UInterchangeFactoryBaseNode* FactoryNode);
-	TObjectPtr<UObject> GetReferencedObjectForNode(const UInterchangeFactoryBaseNode* FactoryNode) const;
 	TObjectPtr<UInterchangeFactoryBase> GetFactoryForNode(const UInterchangeFactoryBaseNode* FactoryNode) const;
+	TArray<TObjectPtr<UObject>> GetObjectInstancesForFactoryNode(const UInterchangeFactoryBaseNode* FactoryNode);
+
 
 	const UInterchangeBaseNodeContainer* GetBaseNodeContainer() const;
 	const UInterchangeSceneImportAsset* GetSceneImportAsset() const;
@@ -64,9 +68,9 @@ public:
 	FResetPhaseDelegates ResetDelegates;
 	FResetPhaseDelegates PostResetDelegates;
 
-	TSet<const UInterchangeFactoryBaseNode*> FilteredNodes;
 
 	TUniquePtr<FInterchangeResetContextData> ResetContextData;
+	TSet<const UInterchangeFactoryBaseNode*> FilteredNodes;
 private:
 	TObjectPtr<const UInterchangeSceneImportAsset> SceneImportAsset;
 	TObjectPtr<UInterchangeResultsContainer> ResultsContainer;
@@ -79,7 +83,6 @@ class FInterchangeReset
 {
 public:
 	static INTERCHANGEENGINE_API void ExecuteReset(FInterchangeResetParameters& ResetObjectParameters);
-
 };
 
 namespace UE::Interchange::InterchangeReset
@@ -89,6 +92,7 @@ namespace UE::Interchange::InterchangeReset
 	namespace Constants
 	{
 		INTERCHANGEENGINE_API extern const FString SceneImportAssetPathKey;
+		INTERCHANGEENGINE_API extern const FString FactoryNodeUidPathKey;
 	}
 
 }
