[BUGFIX] Interchange - Import Dialog - Section has inconsistent visibility
- The required information is refetched through a delegate.

#rnx

diff --git a/Engine/Plugins/Interchange/Editor/Source/Pipelines/Private/InterchangeEditorPipelineDetails.cpp b/Engine/Plugins/Interchange/Editor/Source/Pipelines/Private/InterchangeEditorPipelineDetails.cpp
--- a/Engine/Plugins/Interchange/Editor/Source/Pipelines/Private/InterchangeEditorPipelineDetails.cpp
+++ b/Engine/Plugins/Interchange/Editor/Source/Pipelines/Private/InterchangeEditorPipelineDetails.cpp
@@ -16,6 +16,8 @@
 
 #define LOCTEXT_NAMESPACE "InterchangeEditorPipelineDetails"
 
+FOnGatherConflictAndExtraInfo FInterchangePipelineBaseDetailsCustomization::OnGatherConflictAndExtraInfo;
+
 FInterchangePipelineBaseDetailsCustomization::FInterchangePipelineBaseDetailsCustomization()
 {
 	InterchangePipeline = nullptr;
@@ -26,6 +28,7 @@ void FInterchangePipelineBaseDetailsCustomization::RefreshCustomDetail()
 {
 	if (CachedDetailBuilder)
 	{
+		ReCollectConflictAndExtraInfo();
 		CachedDetailBuilder->ForceRefreshDetails();
 	}
 }
@@ -830,6 +833,15 @@ void FInterchangePipelineBaseDetailsCustomization::CustomizeDetails(IDetailLayou
 	}
 }
 
+void FInterchangePipelineBaseDetailsCustomization::ReCollectConflictAndExtraInfo()
+{
+	TArray<FInterchangeConflictInfo> ConflictInfo;
+	TMap<FString, FString> ExtraInfo;
+	OnGatherConflictAndExtraInfo.ExecuteIfBound(ConflictInfo, ExtraInfo);
+	SetConflictsInfo(ConflictInfo);
+	SetExtraInformation(ExtraInfo);
+}
+
 FInterchangeBaseNodeDetailsCustomization::FInterchangeBaseNodeDetailsCustomization()
 {
 	InterchangeBaseNode = nullptr;
diff --git a/Engine/Plugins/Interchange/Editor/Source/Pipelines/Private/InterchangeEditorPipelineDetails.h b/Engine/Plugins/Interchange/Editor/Source/Pipelines/Private/InterchangeEditorPipelineDetails.h
--- a/Engine/Plugins/Interchange/Editor/Source/Pipelines/Private/InterchangeEditorPipelineDetails.h
+++ b/Engine/Plugins/Interchange/Editor/Source/Pipelines/Private/InterchangeEditorPipelineDetails.h
@@ -7,6 +7,7 @@
 #include "UObject/WeakObjectPtr.h"
 #include "Templates/SharedPointer.h"
 #include "Types/SlateEnums.h"
+#include "Delegates/DelegateCombinations.h"
 
 class SWidget;
 
@@ -26,6 +27,10 @@ namespace UE
 	}
 }
 
+using FInterchangeConflictInfoArray = TArray<FInterchangeConflictInfo>;
+using FInterchangeExtraInfoMap = TMap<FString, FString>;
+DECLARE_DELEGATE_TwoParams(FOnGatherConflictAndExtraInfo, FInterchangeConflictInfoArray&, FInterchangeExtraInfoMap&);
+
 class FInterchangePipelineBaseDetailsCustomization : public IDetailCustomization
 {
 public:
@@ -43,6 +48,10 @@ public:
 	virtual void CustomizeDetails(IDetailLayoutBuilder& DetailBuilder) override;
 	/** End IDetailCustomization interface */
 
+
+	static void ReCollectConflictAndExtraInfo();
+	static FOnGatherConflictAndExtraInfo OnGatherConflictAndExtraInfo;
+
 private:
 	struct FInternalPropertyData
 	{
diff --git a/Engine/Plugins/Interchange/Editor/Source/Pipelines/Private/SInterchangePipelineConfigurationDialog.cpp b/Engine/Plugins/Interchange/Editor/Source/Pipelines/Private/SInterchangePipelineConfigurationDialog.cpp
--- a/Engine/Plugins/Interchange/Editor/Source/Pipelines/Private/SInterchangePipelineConfigurationDialog.cpp
+++ b/Engine/Plugins/Interchange/Editor/Source/Pipelines/Private/SInterchangePipelineConfigurationDialog.cpp
@@ -371,6 +371,24 @@ FReply SInterchangePipelineConfigurationDialog::OnEditTranslatorSettings()
 	return FReply::Handled();
 }
 
+void SInterchangePipelineConfigurationDialog::GatherConflictAndExtraInfo(TArray<FInterchangeConflictInfo>& ConflictInfo, TMap<FString, FString>& ExtraInfo)
+{
+	if (CurrentSelectedPipelineItem.IsValid())
+	{
+		TSharedPtr<FInterchangePipelineItemType> PipelineItem = CurrentSelectedPipelineItem.Pin();
+		if (PipelineItem->ReimportObject)
+		{
+			ConflictInfo = PipelineItem->Pipeline->GetConflictInfos(PipelineItem->ReimportObject, PipelineItem->Container, PipelineItem->SourceData);
+		}
+
+		const UInterchangeSourceNode* SourceNode = UInterchangeSourceNode::GetUniqueInstance(PipelineItem->Container);
+		if (SourceNode)
+		{
+			SourceNode->GetExtraInformation(ExtraInfo);
+		}
+	}
+}
+
 TSharedRef<SBox> SInterchangePipelineConfigurationDialog::SpawnPipelineConfiguration()
 {
 	const ISlateStyle* InterchangeEditorPipelineStyle = FSlateStyleRegistry::FindSlateStyle("InterchangeEditorPipelineStyle");
@@ -1184,6 +1202,8 @@ void SInterchangePipelineConfigurationDialog::Construct(const FArguments& InArgs
 			RefreshCardsViewList();
 		}
 	}
+
+	FInterchangePipelineBaseDetailsCustomization::OnGatherConflictAndExtraInfo.BindRaw(this, &SInterchangePipelineConfigurationDialog::GatherConflictAndExtraInfo);
 }
 
 bool SInterchangePipelineConfigurationDialog::IsPropertyVisible(const FPropertyAndParent& PropertyAndParent) const
@@ -1388,6 +1408,9 @@ void SInterchangePipelineConfigurationDialog::ClosePipelineConfiguration(const E
 		}
 	}
 	OwnerWindow = nullptr;
+
+
+	FInterchangePipelineBaseDetailsCustomization::OnGatherConflictAndExtraInfo.Unbind();
 }
 
 FReply SInterchangePipelineConfigurationDialog::OnKeyDown(const FGeometry& MyGeometry, const FKeyEvent& InKeyEvent)
@@ -1533,6 +1556,7 @@ void SInterchangePipelineConfigurationDialog::OnPipelineSelectionChanged(TShared
 	{
 		CurrentSelectedPipeline = InItem->Pipeline;
 	}
+	CurrentSelectedPipelineItem = InItem;
 	SetEditPipeline(InItem.Get());
 	
 	if (CurrentSelectedPipeline)
diff --git a/Engine/Plugins/Interchange/Editor/Source/Pipelines/Private/SInterchangePipelineConfigurationDialog.h b/Engine/Plugins/Interchange/Editor/Source/Pipelines/Private/SInterchangePipelineConfigurationDialog.h
--- a/Engine/Plugins/Interchange/Editor/Source/Pipelines/Private/SInterchangePipelineConfigurationDialog.h
+++ b/Engine/Plugins/Interchange/Editor/Source/Pipelines/Private/SInterchangePipelineConfigurationDialog.h
@@ -120,6 +120,7 @@ private:
 	void SetEditPipeline(FInterchangePipelineItemType* PipelineItemToEdit);
 	FReply OnEditTranslatorSettings();
 
+	void GatherConflictAndExtraInfo(TArray<FInterchangeConflictInfo>& ConflictInfo, TMap<FString, FString>& ExtraInfo);
 private:
 	TWeakPtr< SWindow > OwnerWindow;
 	double OriginalMinWindowSize = 0.0;
@@ -206,4 +207,5 @@ private:
 
 	FName CurrentStackName = NAME_None;
 	TObjectPtr<UInterchangePipelineBase> CurrentSelectedPipeline = nullptr;
+	TWeakPtr<FInterchangePipelineItemType> CurrentSelectedPipelineItem = nullptr;
 };
