[BUGFIX] Interchange Editor Icon Path fix
- Fixed the path

#rnx

diff --git a/Engine/Plugins/Interchange/Editor/Source/Pipelines/Private/InterchangeEditorPipelineStyle.h b/Engine/Plugins/Interchange/Editor/Source/Pipelines/Private/InterchangeEditorPipelineStyle.h
--- a/Engine/Plugins/Interchange/Editor/Source/Pipelines/Private/InterchangeEditorPipelineStyle.h
+++ b/Engine/Plugins/Interchange/Editor/Source/Pipelines/Private/InterchangeEditorPipelineStyle.h
@@ -26,6 +26,7 @@ public:
 
 		const FString BaseDir = IPluginManager::Get().FindPlugin("InterchangeEditor")->GetBaseDir();
 		SetContentRoot(BaseDir / TEXT("Content"));
+		SetCoreContentRoot(FPaths::EngineContentDir() / TEXT("Slate"));
 
 		//GraphInspectorIcon icons
 		FSlateImageBrush* LodBrush16 = new FSlateImageBrush(RootToContentDir(TEXT("Resources/Interchange_Lod_Icon_16"), TEXT(".png")), Icon16x16);
@@ -49,7 +50,8 @@ public:
 
 
 		//InterchangeAssetCard icons
-		FSlateImageBrush* AlertSolidWarningBrush20 = new FSlateImageBrush(RootToContentDir(TEXT("Old/Tiles/Outer/alertSolid"), TEXT(".png")), Icon20x20, FStyleColors::Warning);
+		UE_LOG(LogTemp, Display, TEXT("AlertSolidPath: %s"), *RootToCoreContentDir(TEXT("Old/Tiles/Outer/alertSolid"), TEXT(".png")));
+		FSlateImageBrush* AlertSolidWarningBrush20 = new FSlateImageBrush(RootToCoreContentDir(TEXT("Old/Tiles/Outer/alertSolid"), TEXT(".png")), Icon20x20, FStyleColors::Warning);
 		Set("InterchangeAssetCardIcon.Warning", AlertSolidWarningBrush20);
 
 		FSlateImageBrush* SidePanelRightBrush20 = new FSlateImageBrush(RootToContentDir(TEXT("Resources/SidePanelRight"), TEXT(".png")), Icon20x20);
