Removing Temp Logs

#rnx

diff --git a/Engine/Plugins/Interchange/Editor/Source/Pipelines/Private/InterchangeEditorPipelineStyle.h b/Engine/Plugins/Interchange/Editor/Source/Pipelines/Private/InterchangeEditorPipelineStyle.h
--- a/Engine/Plugins/Interchange/Editor/Source/Pipelines/Private/InterchangeEditorPipelineStyle.h
+++ b/Engine/Plugins/Interchange/Editor/Source/Pipelines/Private/InterchangeEditorPipelineStyle.h
@@ -48,9 +48,7 @@ public:
 		FSlateImageBrush* TranslatorSettingsBrush16 = new FSlateImageBrush(RootToContentDir(TEXT("Resources/Interchange_TranslatorSettings_16"), TEXT(".png")), Icon16x16);
 		Set("PipelineConfigurationIcon.TranslatorSettings", TranslatorSettingsBrush16);
 
-
 		//InterchangeAssetCard icons
-		UE_LOG(LogTemp, Display, TEXT("AlertSolidPath: %s"), *RootToCoreContentDir(TEXT("Old/Tiles/Outer/alertSolid"), TEXT(".png")));
 		FSlateImageBrush* AlertSolidWarningBrush20 = new FSlateImageBrush(RootToCoreContentDir(TEXT("Old/Tiles/Outer/alertSolid"), TEXT(".png")), Icon20x20, FStyleColors::Warning);
 		Set("InterchangeAssetCardIcon.Warning", AlertSolidWarningBrush20);
 
