[BUGFIX] Interchange window gets bigger on every usage
- Added a fix so that the size that gets stored in config is unscaled.

#rnx

diff --git a/Engine/Plugins/Interchange/Editor/Source/Pipelines/Private/InterchangePipelineConfigurationGeneric.cpp b/Engine/Plugins/Interchange/Editor/Source/Pipelines/Private/InterchangePipelineConfigurationGeneric.cpp
--- a/Engine/Plugins/Interchange/Editor/Source/Pipelines/Private/InterchangePipelineConfigurationGeneric.cpp
+++ b/Engine/Plugins/Interchange/Editor/Source/Pipelines/Private/InterchangePipelineConfigurationGeneric.cpp
@@ -44,8 +44,16 @@ namespace UE::Interchange::Private
 			GConfig->GetDouble(TEXT("InterchangeImportDialogOptions"), *SizeYUid, OutClientSize.Y, GEditorPerProjectIni);
 		}
 	}
-	void BackupClientSize(const FVector2D& ClientSize, const FString& WindowsUniqueID)
+	void BackupClientSize(SWindow* Window, const FString& WindowsUniqueID)
 	{
+		if (!Window)
+		{
+			return;
+		}
+
+		// Need to convert it back so that on creating the window, it will appropriately adjust for DPI scale.
+		const FVector2D ClientSize = Window->GetClientSizeInScreen() / Window->GetDPIScaleFactor();
+
 		FString SizeXUid = WindowsUniqueID + GetSizeXUid();
 		FString SizeYUid = WindowsUniqueID + GetSizeYUid();
 		GConfig->SetDouble(TEXT("InterchangeImportDialogOptions"), *SizeXUid, ClientSize.X, GEditorPerProjectIni);
@@ -92,8 +100,7 @@ EInterchangePipelineConfigurationDialogResult UInterchangePipelineConfigurationG
 
 	Window->SetOnWindowClosed(FOnWindowClosed::CreateLambda([&WindowDialogUid](const TSharedRef<SWindow>& WindowMoved)
 		{
-			FVector2D LastClientSize = WindowMoved->GetClientSizeInScreen();
-			BackupClientSize(LastClientSize, WindowDialogUid);
+			BackupClientSize(&WindowMoved.Get(), WindowDialogUid);
 		}));
 
 	FSlateApplication::Get().AddModalWindow(Window, ParentWindow, false);
@@ -128,7 +135,6 @@ EInterchangePipelineConfigurationDialogResult UInterchangePipelineConfigurationG
 	FString WindowDialogUid = TEXT("ImportSceneDialog");
 	FVector2D ClientSize = DefaultClientSize();
 	RestoreClientSize(ClientSize, WindowDialogUid);
-
 	TSharedRef<SWindow> Window = SNew(SWindow)
 		.ClientSize(ClientSize)
 		.MinWidth(GetMinSizeX())
@@ -151,8 +157,7 @@ EInterchangePipelineConfigurationDialogResult UInterchangePipelineConfigurationG
 
 	Window->SetOnWindowClosed(FOnWindowClosed::CreateLambda([&WindowDialogUid](const TSharedRef<SWindow>& WindowMoved)
 		{
-			FVector2D LastClientSize = WindowMoved->GetClientSizeInScreen();
-			BackupClientSize(LastClientSize, WindowDialogUid);
+			BackupClientSize(&WindowMoved.Get(), WindowDialogUid);
 		}));
 
 	FSlateApplication::Get().AddModalWindow(Window, ParentWindow, false);
@@ -213,8 +218,7 @@ EInterchangePipelineConfigurationDialogResult UInterchangePipelineConfigurationG
 
 	Window->SetOnWindowClosed(FOnWindowClosed::CreateLambda([&WindowDialogUid](const TSharedRef<SWindow>& WindowMoved)
 		{
-			FVector2D LastClientSize = WindowMoved->GetClientSizeInScreen();
-			BackupClientSize(LastClientSize, WindowDialogUid);
+			BackupClientSize(&WindowMoved.Get(), WindowDialogUid);
 		}));
 
 	FSlateApplication::Get().AddModalWindow(Window, ParentWindow, false);
