[BUGFIX] Interchange Import Dialog
- Clear undo history while closing the interchange import dialog.

#rnx

diff --git a/Engine/Plugins/Interchange/Editor/Source/Pipelines/Private/SInterchangePipelineConfigurationDialog.cpp b/Engine/Plugins/Interchange/Editor/Source/Pipelines/Private/SInterchangePipelineConfigurationDialog.cpp
--- a/Engine/Plugins/Interchange/Editor/Source/Pipelines/Private/SInterchangePipelineConfigurationDialog.cpp
+++ b/Engine/Plugins/Interchange/Editor/Source/Pipelines/Private/SInterchangePipelineConfigurationDialog.cpp
@@ -5,6 +5,8 @@
 
 #include "DetailsViewArgs.h"
 #include "Dialog/SCustomDialog.h"
+#include "Editor.h"
+#include "Editor/Transactor.h"
 #include "Framework/MultiBox/MultiBoxBuilder.h"
 #include "Framework/Commands/UICommandList.h"
 #include "Framework/Views/TableViewMetadata.h"
@@ -1409,6 +1411,11 @@ void SInterchangePipelineConfigurationDialog::ClosePipelineConfiguration(const E
 	}
 	OwnerWindow = nullptr;
 
+	// Discard undo history as these transactions should only be valid for the lifetime of the import dialog.
+	if ((GEditor != nullptr) && (GEditor->Trans != nullptr))
+	{
+		GEditor->Trans->Reset(LOCTEXT("DiscardHistoryReason", "Discard undo history after closing Interchange Import Dialog."));
+	}
 
 	FInterchangePipelineBaseDetailsCustomization::OnGatherConflictAndExtraInfo.Unbind();
 }
