Don't initialize module when cooking

diff --git a/Engine/Plugins/Interchange/Editor/Source/InterchangeEditor/Private/InterchangeEditorModule.cpp b/Engine/Plugins/Interchange/Editor/Source/InterchangeEditor/Private/InterchangeEditorModule.cpp
--- a/Engine/Plugins/Interchange/Editor/Source/InterchangeEditor/Private/InterchangeEditorModule.cpp
+++ b/Engine/Plugins/Interchange/Editor/Source/InterchangeEditor/Private/InterchangeEditorModule.cpp
@@ -99,6 +99,11 @@ bool FInterchangeEditorModule::IsAvailable()
 
 void FInterchangeEditorModule::StartupModule()
 {
+	if (IsRunningCookCommandlet())
+	{
+		return;
+	}
+
 	using namespace UE::Interchange;
 
 	auto RegisterItems = [this]()
