[Interchange] Removed unused icon from Interchange Editor Pipeline Style

#rnx

diff --git a/Engine/Plugins/Interchange/Editor/Source/Pipelines/Private/InterchangeEditorPipelineStyle.h b/Engine/Plugins/Interchange/Editor/Source/Pipelines/Private/InterchangeEditorPipelineStyle.h
--- a/Engine/Plugins/Interchange/Editor/Source/Pipelines/Private/InterchangeEditorPipelineStyle.h
+++ b/Engine/Plugins/Interchange/Editor/Source/Pipelines/Private/InterchangeEditorPipelineStyle.h
@@ -26,7 +26,6 @@ public:
 
 		const FString BaseDir = IPluginManager::Get().FindPlugin("InterchangeEditor")->GetBaseDir();
 		SetContentRoot(BaseDir / TEXT("Content"));
-		SetCoreContentRoot(FPaths::EngineContentDir() / TEXT("Slate"));
 
 		//GraphInspectorIcon icons
 		FSlateImageBrush* LodBrush16 = new FSlateImageBrush(RootToContentDir(TEXT("Resources/Interchange_Lod_Icon_16"), TEXT(".png")), Icon16x16);
@@ -48,10 +47,6 @@ public:
 		FSlateImageBrush* TranslatorSettingsBrush16 = new FSlateImageBrush(RootToContentDir(TEXT("Resources/Interchange_TranslatorSettings_16"), TEXT(".png")), Icon16x16);
 		Set("PipelineConfigurationIcon.TranslatorSettings", TranslatorSettingsBrush16);
 
-		//InterchangeAssetCard icons
-		FSlateImageBrush* AlertSolidWarningBrush20 = new FSlateImageBrush(RootToCoreContentDir(TEXT("Old/Tiles/Outer/alertSolid"), TEXT(".png")), Icon20x20, FStyleColors::Warning);
-		Set("InterchangeAssetCardIcon.Warning", AlertSolidWarningBrush20);
-
 		FSlateImageBrush* SidePanelRightBrush20 = new FSlateImageBrush(RootToContentDir(TEXT("Resources/SidePanelRight"), TEXT(".png")), Icon20x20);
 		Set("PipelineConfigurationIcon.SidePanelRight", SidePanelRightBrush20);
 
