[Interchange] Fix the dialog to show conflict when importing. The conflict were limited to re-import but we need it if we select a skeleton and it doesn't match.

#rnx

diff --git a/Engine/Plugins/Interchange/Editor/Source/Pipelines/Private/InterchangeEditorPipelineDetails.cpp b/Engine/Plugins/Interchange/Editor/Source/Pipelines/Private/InterchangeEditorPipelineDetails.cpp
--- a/Engine/Plugins/Interchange/Editor/Source/Pipelines/Private/InterchangeEditorPipelineDetails.cpp
+++ b/Engine/Plugins/Interchange/Editor/Source/Pipelines/Private/InterchangeEditorPipelineDetails.cpp
@@ -285,7 +285,7 @@ void FInterchangePipelineBaseDetailsCustomization::AddConflictSection()
 	}
 
 	TArray<FInterchangeConflictInfo> ConflictInfos = ConflictInfosStack.Pop(EAllowShrinking::No);
-	if (!InterchangePipeline->IsReimportContext() || ConflictInfos.Num() == 0)
+	if (ConflictInfos.Num() == 0)
 	{
 		return;
 	}
diff --git a/Engine/Plugins/Interchange/Editor/Source/Pipelines/Private/SInterchangePipelineConfigurationDialog.cpp b/Engine/Plugins/Interchange/Editor/Source/Pipelines/Private/SInterchangePipelineConfigurationDialog.cpp
--- a/Engine/Plugins/Interchange/Editor/Source/Pipelines/Private/SInterchangePipelineConfigurationDialog.cpp
+++ b/Engine/Plugins/Interchange/Editor/Source/Pipelines/Private/SInterchangePipelineConfigurationDialog.cpp
@@ -215,10 +215,7 @@ void SInterchangePipelineConfigurationDialog::SetEditPipeline(FInterchangePipeli
 	if (PipelineItemToEdit)
 	{
 		PipelineItemToEdit->ConflictInfos.Reset();
-		if (PipelineItemToEdit->ReimportObject)
-		{
-			PipelineItemToEdit->ConflictInfos = PipelineItemToEdit->Pipeline->GetConflictInfos(PipelineItemToEdit->ReimportObject, PipelineItemToEdit->Container, PipelineItemToEdit->SourceData);
-		}
+		PipelineItemToEdit->ConflictInfos = PipelineItemToEdit->Pipeline->GetConflictInfos(PipelineItemToEdit->ReimportObject, PipelineItemToEdit->Container, PipelineItemToEdit->SourceData);
 		FInterchangePipelineBaseDetailsCustomization::SetConflictsInfo(PipelineItemToEdit->ConflictInfos);
 
 		//Acquire ExtraInformation from SourceNode and pass it to FInterchangePipelineBaseDetailsCustomization:
@@ -378,10 +375,7 @@ void SInterchangePipelineConfigurationDialog::GatherConflictAndExtraInfo(TArray<
 	if (CurrentSelectedPipelineItem.IsValid())
 	{
 		TSharedPtr<FInterchangePipelineItemType> PipelineItem = CurrentSelectedPipelineItem.Pin();
-		if (PipelineItem->ReimportObject)
-		{
-			ConflictInfo = PipelineItem->Pipeline->GetConflictInfos(PipelineItem->ReimportObject, PipelineItem->Container, PipelineItem->SourceData);
-		}
+		ConflictInfo = PipelineItem->Pipeline->GetConflictInfos(PipelineItem->ReimportObject, PipelineItem->Container, PipelineItem->SourceData);
 
 		const UInterchangeSourceNode* SourceNode = UInterchangeSourceNode::GetUniqueInstance(PipelineItem->Container);
 		if (SourceNode)
@@ -1506,10 +1500,7 @@ void SInterchangePipelineConfigurationDialog::RefreshStack(bool bStackSelectionC
 			for (TSharedPtr<FInterchangePipelineItemType>& PiplineListViewItem : PipelineListViewItems)
 			{
 				PiplineListViewItem->ConflictInfos.Reset();
-				if (PiplineListViewItem->ReimportObject)
-				{
-					PiplineListViewItem->ConflictInfos = PiplineListViewItem->Pipeline->GetConflictInfos(PiplineListViewItem->ReimportObject, PiplineListViewItem->Container, PiplineListViewItem->SourceData);
-				}
+				PiplineListViewItem->ConflictInfos = PiplineListViewItem->Pipeline->GetConflictInfos(PiplineListViewItem->ReimportObject, PiplineListViewItem->Container, PiplineListViewItem->SourceData);
 			}
 		}
 	}
diff --git a/Engine/Plugins/Interchange/Runtime/Source/Pipelines/Private/InterchangeGenericAssetsPipeline.cpp b/Engine/Plugins/Interchange/Runtime/Source/Pipelines/Private/InterchangeGenericAssetsPipeline.cpp
--- a/Engine/Plugins/Interchange/Runtime/Source/Pipelines/Private/InterchangeGenericAssetsPipeline.cpp
+++ b/Engine/Plugins/Interchange/Runtime/Source/Pipelines/Private/InterchangeGenericAssetsPipeline.cpp
@@ -926,7 +926,7 @@ void UInterchangeGenericAssetsPipeline::CreateSkeletonConflict(USkeleton* Specif
 			return;
 		}
 
-		if (!SpecifiedSkeleton && !SkeletalMesh->GetSkeleton())
+		if (!SpecifiedSkeleton && !(SkeletalMesh && SkeletalMesh->GetSkeleton()))
 		{
 			return;
 		}
@@ -1036,7 +1036,7 @@ TArray<FInterchangeConflictInfo> UInterchangeGenericAssetsPipeline::GetConflictI
 		CreateMaterialConflict(StaticMesh, SkeletalMesh, TransientBaseNodeContainer);
 	}
 
-	if (SpecifiedSkeleton && SkeletalMesh)
+	if (SpecifiedSkeleton || SkeletalMesh)
 	{
 		CreateSkeletonConflict(SpecifiedSkeleton, SkeletalMesh, TransientBaseNodeContainer);
 	}
