[Interchange UI]  Added a standalone SInterchangeTranslatorSettingsDialog

#rnx

diff --git a/Engine/Plugins/Interchange/Editor/Source/Pipelines/Private/SInterchangePipelineConfigurationDialog.cpp b/Engine/Plugins/Interchange/Editor/Source/Pipelines/Private/SInterchangePipelineConfigurationDialog.cpp
--- a/Engine/Plugins/Interchange/Editor/Source/Pipelines/Private/SInterchangePipelineConfigurationDialog.cpp
+++ b/Engine/Plugins/Interchange/Editor/Source/Pipelines/Private/SInterchangePipelineConfigurationDialog.cpp
@@ -29,6 +29,7 @@
 #include "Nodes/InterchangeSourceNode.h"
 #include "PropertyEditorModule.h"
 #include "SInterchangeGraphInspectorWindow.h"
+#include "SInterchangeTranslatorSettingsDialog.h"
 #include "SPrimaryButton.h"
 #include "Styling/SlateIconFinder.h"
 #include "Styling/StyleColors.h"
@@ -236,136 +237,47 @@ FReply SInterchangePipelineConfigurationDialog::OnEditTranslatorSettings()
 	{
 		return FReply::Handled();
 	}
-	UInterchangeTranslatorSettings* TranslatorSettingsCDO = TranslatorSettings->GetClass()->GetDefaultObject<UInterchangeTranslatorSettings>();
-	UInterchangeTranslatorSettings* OriginalTranslatorSettings = DuplicateObject<UInterchangeTranslatorSettings>(TranslatorSettings, GetTransientPackage());
 
-	FPropertyEditorModule& PropertyEditorModule = FModuleManager::GetModuleChecked<FPropertyEditorModule>("PropertyEditor");
-	FDetailsViewArgs DetailsViewArgs;
-	DetailsViewArgs.bAllowSearch = true;
-	DetailsViewArgs.bShowPropertyMatrixButton = false;
-	DetailsViewArgs.bShowSectionSelector = false;
-	DetailsViewArgs.bAllowMultipleTopLevelObjects = false;
-	DetailsViewArgs.bShowModifiedPropertiesOption = false;
-	DetailsViewArgs.bShowKeyablePropertiesOption = false;
-	DetailsViewArgs.bShowAnimatedPropertiesOption = false;
-	DetailsViewArgs.bShowHiddenPropertiesWhilePlayingOption = false;
-	DetailsViewArgs.NameAreaSettings = FDetailsViewArgs::HideNameArea;
-	TSharedRef<IDetailsView> TranslatorSettingsDetailsView = PropertyEditorModule.CreateDetailView(DetailsViewArgs);
+	TSharedRef<SInterchangeTranslatorSettingsDialog> OptionsDialog =
+		SNew(SInterchangeTranslatorSettingsDialog)
+		.WindowArguments(SWindow::FArguments()
+			.IsTopmostWindow(false)
+			.MinWidth(500)
+			.MinHeight(400)
+			.ClientSize(FVector2f{500, 400})
+			.SizingRule(ESizingRule::UserSized))
+		.TranslatorSettings(TranslatorSettings);
 
-	bool bIsSettingsModified = false;
-	TranslatorSettingsDetailsView->OnFinishedChangingProperties().AddLambda([&bIsSettingsModified](const FPropertyChangedEvent& PropertyChangedEvent)
-		{
-			bIsSettingsModified = true;
-		});
-	TranslatorSettingsDetailsView->SetObject(TranslatorSettings);
 
-	TSharedRef<SCustomDialog> OptionsDialog =
-		SNew(SCustomDialog)
-		.Title(LOCTEXT("OptionsDialogTitle", "Translator Settings"))
-		.WindowArguments(SWindow::FArguments()
-			.IsTopmostWindow(true)
-			.MinWidth(400)
-			.MaxHeight(800)
-			.SizingRule(ESizingRule::Autosized))
-		.HAlignContent(HAlign_Fill)
-		.VAlignContent(VAlign_Fill)
-		.HAlignIcon(HAlign_Left)
-		.VAlignIcon(VAlign_Top)
-		.UseScrollBox(true)
-		.RootPadding(FMargin(16.f, 16.f))
-		.ButtonAreaPadding(FMargin(0.f, 16.f, 0.f, 0.f))
-		.Content()
-		[
-			SNew(SBorder)
-			.Padding(FMargin(0.f, 4.f, 0.f, 0.f))
-			.BorderImage(FAppStyle::GetBrush("ToolPanel.DarkGroupBorder"))
-			[
-				SNew(SVerticalBox)
-				+ SVerticalBox::Slot()
-				.HAlign(HAlign_Fill)
-				.VAlign(VAlign_Fill)
-				.FillHeight(1.0)
-				[
-					TranslatorSettingsDetailsView
-				]
-			]
-			
-		]
-		.AutoCloseOnButtonPress(false)
-		.BeforeButtons()
-		[
-			SNew(SHorizontalBox)
-			+ SHorizontalBox::Slot()
-			.Padding(0.f, 0.f)
-			.AutoWidth()
-			[
-				SNew(SButton)
-				.Text(LOCTEXT("DialogButtonResetToDefault", "Reset to Default"))
-				.OnClicked_Lambda([this, &bIsSettingsModified, &TranslatorSettingsCDO, &OriginalTranslatorSettings]()
-					{
-						if (TranslatorSettingsCDO)
-						{
-							//Save the CDO settings
-							TranslatorSettingsCDO->SaveSettings();
-							//load the CDO settings
-							TranslatorSettings->LoadSettings();
-							bIsSettingsModified = true;
-						}
-						return FReply::Handled();
-					})
-			]
-			+ SHorizontalBox::Slot()
-			.FillWidth(1.0f)
-			[
-				SNullWidget::NullWidget
-			]
-		]
-		.Buttons(
+	OptionsDialog->GetTranslatorSettingsDialogClosed().BindLambda([this](bool bUserResponse, bool bTranslatorSettingsChanged)
 		{
-			SCustomDialog::FButton(LOCTEXT("DialogButtonOk", "Ok"))
-			.SetOnClicked(FSimpleDelegate::CreateLambda([&OptionsDialog]()
-				{
-					OptionsDialog.Get().RequestDestroyWindow();
-				}))
-			.SetPrimary(true),
-			SCustomDialog::FButton(LOCTEXT("DialogButtonCancel", "Cancel"))
-			.SetOnClicked(FSimpleDelegate::CreateLambda([this, &bIsSettingsModified, &OptionsDialog, &OriginalTranslatorSettings]()
-				{
-					if (OriginalTranslatorSettings)
-					{
-						OriginalTranslatorSettings->SaveSettings();
-					}
-					//Revert the settings to what it was before loading the dialog
-					bIsSettingsModified = false;
-					TranslatorSettings->LoadSettings();
-					OptionsDialog.Get().RequestDestroyWindow();
-				}))
-		});
+			if (!bTranslatorSettingsChanged)
+			{
+				return;
+			}
 
-	OptionsDialog->ShowModal();
+			if (UClass* TranslatorSettingsClass = TranslatorSettings->GetClass())
+			{
+				//Save the config locally before the translation.
+				TranslatorSettings->SaveSettings();
 
-	if (bIsSettingsModified && Translator.IsValid() && TranslatorSettings)
-	{
-		if (UClass* TranslatorSettingsClass = TranslatorSettings->GetClass())
-		{
-			//Save the config locally before the translation.
-			TranslatorSettings->SaveSettings();
+				//Need to Translate the source data
+				FScopedSlowTask Progress(2.f, NSLOCTEXT("SInterchangePipelineConfigurationDialog", "TranslatingSourceFile...", "Translating source file..."));
+				Progress.MakeDialog();
+				Progress.EnterProgressFrame(1.f);
+				//Reset the container
+				BaseNodeContainer->Reset();
 
-			//Need to Translate the source data
-			FScopedSlowTask Progress(2.f, NSLOCTEXT("SInterchangePipelineConfigurationDialog", "TranslatingSourceFile...", "Translating source file..."));
-			Progress.MakeDialog();
-			Progress.EnterProgressFrame(1.f);
-			//Reset the container
-			BaseNodeContainer->Reset();
+				Translator->Translate(*BaseNodeContainer.Get());
 
-			Translator->Translate(*BaseNodeContainer.Get());
+				//Refresh the dialog
+				UpdateStack(CurrentStackName);
 
-			//Refresh the dialog
-			UpdateStack(CurrentStackName);
+				Progress.EnterProgressFrame(1.f);
+			}
+		});
 
-			Progress.EnterProgressFrame(1.f);
-		}
-	}
+	OptionsDialog->ShowModal();
 
 	return FReply::Handled();
 }
diff --git a/Engine/Plugins/Interchange/Editor/Source/Pipelines/Private/SInterchangeTranslatorSettingsDialog.cpp b/Engine/Plugins/Interchange/Editor/Source/Pipelines/Private/SInterchangeTranslatorSettingsDialog.cpp
new file mode 100644
--- /dev/null
+++ b/Engine/Plugins/Interchange/Editor/Source/Pipelines/Private/SInterchangeTranslatorSettingsDialog.cpp
@@ -0,0 +1,153 @@
+// Copyright Epic Games, Inc. All Rights Reserved.
+
+#include "SInterchangeTranslatorSettingsDialog.h"
+
+#include "Framework/Application/SlateApplication.h"
+#include "Framework/Docking/TabManager.h"
+#include "Styling/AppStyle.h"
+#include "Styling/SlateBrush.h"
+#include "Widgets/Images/SImage.h"
+#include "Widgets/Input/SButton.h"
+#include "Widgets/Text/STextBlock.h"
+#include "Widgets/Layout/SBox.h"
+#include "Widgets/Layout/SScrollBox.h"
+#include "Widgets/SBoxPanel.h"
+#include "Widgets/SNullWidget.h"
+
+#include "Modules/ModuleManager.h"
+#include "PropertyEditorModule.h"
+#include "InterchangeTranslatorBase.h"
+
+#define LOCTEXT_NAMESPACE "SInterchangeTranslatorSettingsDialog"
+
+void SInterchangeTranslatorSettingsDialog::Construct(const FArguments& InArgs)
+{
+	TranslatorSettings = InArgs._TranslatorSettings;
+	if (TranslatorSettings.IsValid())
+	{
+		TranslatorSettingsCDO = TranslatorSettings->GetClass()->GetDefaultObject<UInterchangeTranslatorSettings>();
+		OriginalTranslatorSettings = DuplicateObject<UInterchangeTranslatorSettings>(TranslatorSettings.Get(), GetTransientPackage());
+	}
+
+	FPropertyEditorModule& PropertyEditorModule = FModuleManager::GetModuleChecked<FPropertyEditorModule>("PropertyEditor");
+	FDetailsViewArgs DetailsViewArgs;
+	DetailsViewArgs.bAllowSearch = true;
+	DetailsViewArgs.bShowPropertyMatrixButton = false;
+	DetailsViewArgs.bShowSectionSelector = false;
+	DetailsViewArgs.bAllowMultipleTopLevelObjects = false;
+	DetailsViewArgs.bShowModifiedPropertiesOption = false;
+	DetailsViewArgs.bShowKeyablePropertiesOption = false;
+	DetailsViewArgs.bShowAnimatedPropertiesOption = false;
+	DetailsViewArgs.bShowHiddenPropertiesWhilePlayingOption = false;
+	DetailsViewArgs.NameAreaSettings = FDetailsViewArgs::HideNameArea;
+	TSharedRef<IDetailsView> TranslatorSettingsDetailsView = PropertyEditorModule.CreateDetailView(DetailsViewArgs);
+	TranslatorSettingsDetailsView->SetObject(TranslatorSettings.Get());
+
+	TranslatorSettingsDetailsView->OnFinishedChangingProperties().AddLambda([this](const FPropertyChangedEvent& PropertyChangedEvent)
+		{
+			bTranslatorSettingsChanged = true;
+		});
+
+	SWindow::Construct(SWindow::FArguments(InArgs._WindowArguments)
+		.Title(LOCTEXT("TranslatorSettingsTitle", "Translator Settings"))
+		.SupportsMaximize(false)
+		.SupportsMinimize(false)
+		[
+			SNew(SBorder)
+			.BorderImage(FAppStyle::Get().GetBrush( "ToolPanel.GroupBorder" ))
+			[
+				SNew(SVerticalBox)
+				+ SVerticalBox::Slot()
+				.FillHeight(1.0f)
+				[
+					TranslatorSettingsDetailsView
+				]
+				+ SVerticalBox::Slot()
+				.AutoHeight()
+				[
+					SNew(SHorizontalBox)
+						+ SHorizontalBox::Slot()
+						.AutoWidth()
+						.HAlign(HAlign_Left)
+						.VAlign(VAlign_Center)
+						[
+							SNew(SButton)
+								.Text(LOCTEXT("TranslatorSettings_ResetToDefault", "Reset To Default"))
+								.OnClicked_Lambda([this]()
+									{
+										if (TranslatorSettingsCDO)
+										{
+											//Save the CDO settings
+											TranslatorSettingsCDO->SaveSettings();
+											// Restore the CDO Settings
+											TranslatorSettings.Get()->LoadSettings();
+											bTranslatorSettingsChanged = true;
+										}
+										return FReply::Handled();
+									})
+						]
+						// Buttons
+						+ SHorizontalBox::Slot()
+						.FillWidth(1.0f)
+						.VAlign(VAlign_Center)
+						.HAlign(HAlign_Right)
+						[
+							SNew(SHorizontalBox)
+								+ SHorizontalBox::Slot()
+								.Padding(FAppStyle::Get().GetMargin("StandardDialog.SlotPadding"))
+								.AutoWidth()
+								[
+									SNew(SButton)
+										.ButtonStyle(&FAppStyle::Get().GetWidgetStyle<FButtonStyle>("PrimaryButton"))
+										.Text(LOCTEXT("TranslatorSettings_OK", "OK"))
+										.OnClicked_Lambda([this]() 
+											{
+												bUserResponse = true;
+												RequestDestroyWindow();
+												return FReply::Handled();
+											})
+								]
+								+ SHorizontalBox::Slot()
+								.Padding(FAppStyle::Get().GetMargin("StandardDialog.SlotPadding"))
+								.AutoWidth()
+								[
+									SNew(SButton)
+										.ButtonStyle(&FAppStyle::Get().GetWidgetStyle<FButtonStyle>("Button"))
+										.Text(LOCTEXT("TranslatorSettings_Cancel", "Cancel"))
+										.OnClicked_Lambda([this]()
+											{
+												if (OriginalTranslatorSettings)
+												{
+													//Save the Original settings
+													OriginalTranslatorSettings->SaveSettings();
+													// Restore the Original Settings
+													TranslatorSettings.Get()->LoadSettings();
+													bTranslatorSettingsChanged = false;
+												}
+
+												bUserResponse = false;
+												RequestDestroyWindow();
+												return FReply::Handled();
+											})
+								]
+						]
+				]
+			]
+		]);
+}
+
+bool SInterchangeTranslatorSettingsDialog::ShowModal()
+{
+	FSlateApplication::Get().AddModalWindow(StaticCastSharedRef<SWindow>(this->AsShared()), FGlobalTabmanager::Get()->GetRootWindow());
+
+	if (OnTranslatorSettingsDialogClosed.IsBound())
+	{
+		OnTranslatorSettingsDialogClosed.Execute(bUserResponse, bTranslatorSettingsChanged);
+	}
+
+	return bUserResponse;
+}
+
+
+
+#undef LOCTEXT_NAMESPACE
\ No newline at end of file
diff --git a/Engine/Plugins/Interchange/Editor/Source/Pipelines/Private/SInterchangeTranslatorSettingsDialog.h b/Engine/Plugins/Interchange/Editor/Source/Pipelines/Private/SInterchangeTranslatorSettingsDialog.h
new file mode 100644
--- /dev/null
+++ b/Engine/Plugins/Interchange/Editor/Source/Pipelines/Private/SInterchangeTranslatorSettingsDialog.h
@@ -0,0 +1,53 @@
+// Copyright Epic Games, Inc. All Rights Reserved.
+
+#pragma once
+
+#include "CoreMinimal.h"
+#include "Styling/AppStyle.h"
+#include "Widgets/SWindow.h"
+#include "UObject/WeakObjectPtr.h"
+#include "Delegates/DelegateCombinations.h"
+
+class UInterchangeTranslatorSettings;
+
+DECLARE_DELEGATE_TwoParams(FOnTranslatorSettingsDialogClosed, bool, bool);
+
+class SInterchangeTranslatorSettingsDialog : public SWindow
+{
+
+public:
+	SLATE_BEGIN_ARGS(SInterchangeTranslatorSettingsDialog)
+		{
+			_AccessibleParams = FAccessibleWidgetData(EAccessibleBehavior::Auto);
+		}
+
+		/*********** Functional ***********/
+		/** Event triggered when the dialog is closed, either because one of the buttons is pressed, or the windows is closed. */
+		SLATE_EVENT(FSimpleDelegate, OnClosed)
+
+		/** Provides default values for SWindow::FArguments not overriden by SCustomDialog. */
+		SLATE_ARGUMENT(SWindow::FArguments, WindowArguments)
+
+		SLATE_ARGUMENT(TWeakObjectPtr<UInterchangeTranslatorSettings>, TranslatorSettings)
+	SLATE_END_ARGS()
+	
+	void Construct(const FArguments& InArgs);
+
+	/** Show a modal dialog. Will block until an input is received.
+	 *  
+	 */
+	bool ShowModal();
+
+	FOnTranslatorSettingsDialogClosed& GetTranslatorSettingsDialogClosed() { return OnTranslatorSettingsDialogClosed; }
+
+private:
+	FOnTranslatorSettingsDialogClosed OnTranslatorSettingsDialogClosed;
+
+	TWeakObjectPtr<UInterchangeTranslatorSettings> TranslatorSettings;
+	UInterchangeTranslatorSettings* OriginalTranslatorSettings = nullptr;
+	UInterchangeTranslatorSettings* TranslatorSettingsCDO = nullptr;
+
+	bool bTranslatorSettingsChanged = false;
+	bool bUserResponse = false;
+
+};
