Interchange UX - Import dialog - clicking below pipeline list makes window blank

#rnx

diff --git a/Engine/Plugins/Interchange/Editor/Source/Pipelines/Private/SInterchangePipelineConfigurationDialog.cpp b/Engine/Plugins/Interchange/Editor/Source/Pipelines/Private/SInterchangePipelineConfigurationDialog.cpp
--- a/Engine/Plugins/Interchange/Editor/Source/Pipelines/Private/SInterchangePipelineConfigurationDialog.cpp
+++ b/Engine/Plugins/Interchange/Editor/Source/Pipelines/Private/SInterchangePipelineConfigurationDialog.cpp
@@ -360,6 +360,7 @@ TSharedRef<SBox> SInterchangePipelineConfigurationDialog::SpawnPipelineConfigura
 		.ListItemsSource(&PipelineListViewItems)
 		.OnGenerateRow(this, &SInterchangePipelineConfigurationDialog::MakePipelineListRowWidget)
 		.OnSelectionChanged(this, &SInterchangePipelineConfigurationDialog::OnPipelineSelectionChanged)
+		.ClearSelectionOnClick(false)
 		.ToolTipText(PipelineListTooltip);
 
 	TSharedPtr<STextComboBox> TextComboBoxPtr;
