Fixed warnings following CL 42255680. See https://horde.devtools.epicgames.com/job/681af32fd406d137bbfde53b?step=b2fc
Root cause: FInterchangePipelineSettingsCacheHandler::ShutdownCacheHandler called too late
Solution: Moved call to FInterchangePipelineSettingsCacheHandler::ShutdownCacheHandler on application pre-exit

#rnx

diff --git a/Engine/Plugins/Interchange/Editor/Source/InterchangeEditor/Private/InterchangeEditorModule.cpp b/Engine/Plugins/Interchange/Editor/Source/InterchangeEditor/Private/InterchangeEditorModule.cpp
--- a/Engine/Plugins/Interchange/Editor/Source/InterchangeEditor/Private/InterchangeEditorModule.cpp
+++ b/Engine/Plugins/Interchange/Editor/Source/InterchangeEditor/Private/InterchangeEditorModule.cpp
@@ -140,6 +140,7 @@ void FInterchangeEditorModule::StartupModule()
 		};
 
 		InterchangeManager.OnPreDestroyInterchangeManager.AddLambda(UnregisterItems);
+		FInterchangePipelineSettingsCacheHandler::InitializeCacheHandler();
 	};
 
 	if (GEngine)
@@ -152,12 +153,11 @@ void FInterchangeEditorModule::StartupModule()
 	}
 
 	FInterchangeResetContextMenuExtender::SetupLevelEditorContextMenuExtender();
-	FInterchangePipelineSettingsCacheHandler::InitializeCacheHandler();
+	FCoreDelegates::OnPreExit.AddStatic(FInterchangePipelineSettingsCacheHandler::ShutdownCacheHandler);
 }
 
 void FInterchangeEditorModule::ShutdownModule()
 {
-	FInterchangePipelineSettingsCacheHandler::ShutdownCacheHandler();
 	FInterchangeResetContextMenuExtender::RemoveLevelEditorContextMenuExtender();
 }
 
diff --git a/Engine/Plugins/Interchange/Editor/Source/InterchangeEditor/Private/InterchangePipelineSettingsCacheHandler.cpp b/Engine/Plugins/Interchange/Editor/Source/InterchangeEditor/Private/InterchangePipelineSettingsCacheHandler.cpp
--- a/Engine/Plugins/Interchange/Editor/Source/InterchangeEditor/Private/InterchangePipelineSettingsCacheHandler.cpp
+++ b/Engine/Plugins/Interchange/Editor/Source/InterchangeEditor/Private/InterchangePipelineSettingsCacheHandler.cpp
@@ -23,7 +23,7 @@ void FInterchangePipelineSettingsCacheHandler::InitializeCacheHandler()
 
 void FInterchangePipelineSettingsCacheHandler::OnAssetRemoved(const FAssetData& RemovedAsset)
 {
-	if (UObject* Asset = RemovedAsset.GetAsset(); !Asset || !Asset->IsA<UInterchangePipelineBase>())
+	if (UClass* AssetClass = RemovedAsset.GetClass(); !AssetClass || !AssetClass->IsChildOf<UInterchangePipelineBase>())
 	{
 		return;
 	}
