Removed #pragma optimize from a InterchangeEditorScriptLibrary.cpp

#rnx

diff --git a/Engine/Plugins/Interchange/Editor/Source/InterchangeEditor/Private/InterchangeEditorScriptLibrary.cpp b/Engine/Plugins/Interchange/Editor/Source/InterchangeEditor/Private/InterchangeEditorScriptLibrary.cpp
--- a/Engine/Plugins/Interchange/Editor/Source/InterchangeEditor/Private/InterchangeEditorScriptLibrary.cpp
+++ b/Engine/Plugins/Interchange/Editor/Source/InterchangeEditor/Private/InterchangeEditorScriptLibrary.cpp
@@ -23,7 +23,7 @@
 
 
 #define LOCTEXT_NAMESPACE "InterchangeEditorScriptLibrary"
-#pragma optimize("",off)
+
 namespace UE::InterchangeEditorScriptLibrary::Private
 {
 	namespace InterchangeReset
@@ -592,5 +592,5 @@ const TArray<AActor*> UInterchangeEditorScriptLibrary::LevelInstanceGetEditableA
 
 	return TArray<AActor*>();
 }
-#pragma optimize("",on)
+
 #undef LOCTEXT_NAMESPACE // "InterchangeResetEditorSubsystem"
\ No newline at end of file
