[BUGFIX] Interchange - reset an imported level working badly with Instance and packed import hierarchy
- Fixes to Interchange reset for the Level Instances and Packed Level Actors to work correctly with mutliple instances of the same level instance.
- Disabled InterchangeReset as the Level Instances and Packed Level Actor is lacking to fully support this feature.

#rnx

diff --git a/Engine/Plugins/Interchange/Editor/Source/InterchangeEditor/Private/AssetDefinition_InterchangeSceneImportAsset.cpp b/Engine/Plugins/Interchange/Editor/Source/InterchangeEditor/Private/AssetDefinition_InterchangeSceneImportAsset.cpp
--- a/Engine/Plugins/Interchange/Editor/Source/InterchangeEditor/Private/AssetDefinition_InterchangeSceneImportAsset.cpp
+++ b/Engine/Plugins/Interchange/Editor/Source/InterchangeEditor/Private/AssetDefinition_InterchangeSceneImportAsset.cpp
@@ -16,6 +16,7 @@
 #include "Engine/World.h"
 #include "Engine/Level.h"
 #include "GameFramework/WorldSettings.h"
+#include "HAL/IConsoleManager.h"
 #include "LevelInstance/LevelInstanceActor.h"
 #include "LevelInstance/LevelInstanceComponent.h"
 #include "Misc/App.h"
@@ -124,6 +125,13 @@ namespace MenuExtension_InterchangeSceneImportAsset
 
 	bool CanExecuteResetScene(const FToolMenuContext& InContext)
 	{
+		IConsoleVariable* CVar = IConsoleManager::Get().FindConsoleVariable(TEXT("Interchange.FeatureFlags.InterchangeReset"), false);
+		const bool bInterchangeResetEnabled = CVar ? CVar->GetBool() : false;
+		if (!bInterchangeResetEnabled)
+		{
+			return false;
+		}
+
 		const UContentBrowserAssetContextMenuContext* Context = UContentBrowserAssetContextMenuContext::FindContextWithAssets(InContext);
 		TArray<UInterchangeSceneImportAsset*> SceneImportAssets = Context->LoadSelectedObjects<UInterchangeSceneImportAsset>();
 
@@ -192,6 +200,9 @@ namespace MenuExtension_InterchangeSceneImportAsset
 				FToolMenuSection& Section = Menu->FindOrAddSection("Interchange");
 				Section.AddDynamicEntry(NAME_None, FNewToolMenuSectionDelegate::CreateLambda([](FToolMenuSection& InSection)
 					{
+						IConsoleVariable* CVar = IConsoleManager::Get().FindConsoleVariable(TEXT("Interchange.FeatureFlags.InterchangeReset"), false);
+						const bool bInterchangeResetEnabled = CVar ? CVar->GetBool() : false;
+						if (bInterchangeResetEnabled)
 						{
 							const TAttribute<FText> Label = LOCTEXT("InterchangeSceneImportAsset_ResetScene", "Reset Scene");
 							const TAttribute<FText> ToolTip = LOCTEXT("InterchangeSceneImportAsset_ResetSceneTooltip", "Reset the scene associated with each selected InterchangeSceneImportAsset to original pipeline properties.");
@@ -203,6 +214,7 @@ namespace MenuExtension_InterchangeSceneImportAsset
 
 							InSection.AddMenuEntry("InterchangeSceneImportAsset_ResetScene", Label, ToolTip, Icon, UIAction);
 						}
+
 						{
 							const TAttribute<FText> Label = LOCTEXT("InterchangeSceneImportAsset_Reimport", "Reimport Scene");
 							const TAttribute<FText> ToolTip = LOCTEXT("InterchangeSceneImportAsset_ReimportTooltip", "Reimport the scene associated with each selected InterchangeSceneImportAsset.");
diff --git a/Engine/Plugins/Interchange/Editor/Source/InterchangeEditor/Private/InterchangeEditorModule.cpp b/Engine/Plugins/Interchange/Editor/Source/InterchangeEditor/Private/InterchangeEditorModule.cpp
--- a/Engine/Plugins/Interchange/Editor/Source/InterchangeEditor/Private/InterchangeEditorModule.cpp
+++ b/Engine/Plugins/Interchange/Editor/Source/InterchangeEditor/Private/InterchangeEditorModule.cpp
@@ -4,6 +4,7 @@
 #include "AssetRegistry/IAssetRegistry.h"
 #include "AssetToolsModule.h"
 #include "Engine/Engine.h"
+#include "HAL/IConsoleManager.h"
 #include "IMessageLogListing.h"
 #include "InterchangeEditorLog.h"
 #include "InterchangeFbxAssetImportDataConverter.h"
@@ -86,6 +87,23 @@ namespace UE::Interchange::InterchangeEditorModule
 		UEditorLoadingSavingSettings* LoadingSavingSettings = GetMutableDefault<UEditorLoadingSavingSettings>();
 		LoadingSavingSettings->bAutoSaveEnable = bOldAutoSaveState;
 	}
+
+
+	FConsoleVariableSinkHandle InterchangeResetCVarSinkHandle;
+	void HandleInterchangeResetEnabledConsoleVarChanged()
+	{
+		if (IConsoleVariable* CVar = IConsoleManager::Get().FindConsoleVariable(TEXT("Interchange.FeatureFlags.InterchangeReset"), false))
+		{
+			if (CVar->GetBool())
+			{
+				FInterchangeResetContextMenuExtender::SetupLevelEditorContextMenuExtender();
+			}
+			else
+			{
+				FInterchangeResetContextMenuExtender::RemoveLevelEditorContextMenuExtender();
+			}
+		}
+	}
 } //ns UE::Interchange::InterchangeEditorModule
 
 FInterchangeEditorModule& FInterchangeEditorModule::Get()
@@ -141,6 +159,12 @@ void FInterchangeEditorModule::StartupModule()
 
 		InterchangeManager.OnPreDestroyInterchangeManager.AddLambda(UnregisterItems);
 		FInterchangePipelineSettingsCacheHandler::InitializeCacheHandler();
+
+		{
+			using namespace UE::Interchange::InterchangeEditorModule;
+			// Register a Console Variable sink to monitor changes in Interchange.FeatureFlags.InterchangeReset
+			InterchangeResetCVarSinkHandle = IConsoleManager::Get().RegisterConsoleVariableSink_Handle(FConsoleCommandDelegate::CreateStatic(&HandleInterchangeResetEnabledConsoleVarChanged));
+		}
 	};
 
 	if (GEngine)
@@ -159,6 +183,11 @@ void FInterchangeEditorModule::StartupModule()
 void FInterchangeEditorModule::ShutdownModule()
 {
 	FInterchangeResetContextMenuExtender::RemoveLevelEditorContextMenuExtender();
+
+	{
+		using namespace UE::Interchange::InterchangeEditorModule;
+		IConsoleManager::Get().UnregisterConsoleVariableSink_Handle(InterchangeResetCVarSinkHandle);
+	}
 }
 
 IMPLEMENT_MODULE(FInterchangeEditorModule, InterchangeEditor)
diff --git a/Engine/Plugins/Interchange/Editor/Source/InterchangeEditor/Private/InterchangeEditorScriptLibrary.cpp b/Engine/Plugins/Interchange/Editor/Source/InterchangeEditor/Private/InterchangeEditorScriptLibrary.cpp
--- a/Engine/Plugins/Interchange/Editor/Source/InterchangeEditor/Private/InterchangeEditorScriptLibrary.cpp
+++ b/Engine/Plugins/Interchange/Editor/Source/InterchangeEditor/Private/InterchangeEditorScriptLibrary.cpp
@@ -15,6 +15,8 @@
 #include "GameFramework/WorldSettings.h"
 #include "GameFramework/Actor.h"
 
+#include "HAL/IConsoleManager.h"
+
 #include "LevelInstance/LevelInstanceActor.h"
 #include "LevelInstance/LevelInstanceComponent.h"
 #include "LevelInstance/LevelInstanceSubsystem.h"
@@ -28,6 +30,11 @@ namespace UE::InterchangeEditorScriptLibrary::Private
 {
 	namespace InterchangeReset
 	{
+		bool IsInterchangeResetAvailable()
+		{
+			IConsoleVariable* CVar = IConsoleManager::Get().FindConsoleVariable(TEXT("Interchange.FeatureFlags.InterchangeReset"), false);
+			return CVar ? CVar->GetBool() : false;
+		}
 
 		ALevelInstance* SpawnLevelInstanceInEditor(FName AssetName)
 		{
@@ -207,14 +214,18 @@ namespace UE::InterchangeEditorScriptLibrary::Private
 			ILevelInstanceInterface* LevelInstance = Cast<ILevelInstanceInterface>(LevelInstanceActor);
 
 			const TSoftObjectPtr<UWorld>& WorldObjectPtr = LevelInstanceActor->GetWorldAsset();
+			FTransform CachedLevelInstanceActorTransform;
 
 			if (const UWorld* World = Cast<UWorld>(WorldObjectPtr.ToSoftObjectPath().TryLoad()))
 			{
 				using namespace UE::Interchange::InterchangeReset;
 
+				const FLevelInstanceID CurrentLevelInstanceID = LevelInstance->GetLevelInstanceID();
 				const bool bWasPrevEditing = LevelInstance->IsEditing();
 				if (!bWasPrevEditing && LevelInstance->CanEnterEdit())
 				{
+					CachedLevelInstanceActorTransform = LevelInstanceActor->GetActorTransform();
+					LevelInstanceActor->SetActorTransform(FTransform::Identity);
 					LevelInstance->EnterEdit(LevelInstanceActor);
 				}
 
@@ -226,8 +237,6 @@ namespace UE::InterchangeEditorScriptLibrary::Private
 					if (ULevel* LoadedLevel = LevelInstance->GetLoadedLevel())
 					{
 						TArray<AActor*> SceneActors = LoadedLevel->Actors;
-						SceneActors.Add(LevelInstanceActor);
-
 						TSet<const UInterchangeFactoryBaseNode*> TempFilteredFactoryNodes;
 						for (AActor* ChildActor : SceneActors)
 						{
@@ -240,16 +249,20 @@ namespace UE::InterchangeEditorScriptLibrary::Private
 				}
 
 				FInterchangeReset::ExecuteReset(ResetParameters);
-#if WITH_EDITOR
-				// Make sure transformation is up to date after registration as its possible LevelInstance actor can get unregistered when editing properties
-				// through Details panel. In this case the ULevelInstanceComponent might not be able to update the ALevelInstanceEditorInstanceActor transform.
-				LevelInstance->GetLevelInstanceComponent()->UpdateEditorInstanceActor();
-#endif
-				GEditor->RedrawAllViewports();
 
 				if (!bWasPrevEditing && LevelInstance->CanExitEdit())
 				{
 					LevelInstance->ExitEdit();
+					if (ULevelInstanceSubsystem* LevelInstanceSubsystem = LevelInstance->GetLevelInstanceSubsystem())
+					{
+						if (ILevelInstanceInterface* LevelInstanceFromID = LevelInstanceSubsystem->GetLevelInstance(CurrentLevelInstanceID))
+						{
+							if (AActor* UpdatedLevelInstanceActor = Cast<AActor>(LevelInstanceFromID))
+							{
+								UpdatedLevelInstanceActor->SetActorTransform(CachedLevelInstanceActorTransform);
+							}
+						}
+					}
 				}
 
 				return true;
@@ -301,8 +314,14 @@ namespace UE::InterchangeEditorScriptLibrary::Private
 }
 
 void UInterchangeEditorScriptLibrary::ResetLevelAsset(UWorld* World)
-{
+{	
 	using namespace UE::InterchangeEditorScriptLibrary::Private::InterchangeReset;
+
+	if (!IsInterchangeResetAvailable())
+	{
+		return;
+	}
+
 	if (UInterchangeSceneImportAsset* SceneImportAsset = GetSceneImportAssetForWorld(World))
 	{
 		ExecuteWorldReset(World, SceneImportAsset);
@@ -311,6 +330,12 @@ void UInterchangeEditorScriptLibrary::ResetLevelAsset(UWorld* World)
 
 void UInterchangeEditorScriptLibrary::ResetSceneImportAsset(UInterchangeSceneImportAsset* SceneImportAsset)
 {
+	using namespace UE::InterchangeEditorScriptLibrary::Private::InterchangeReset;
+	if (!IsInterchangeResetAvailable())
+	{
+		return;
+	}
+
 	TRACE_CPUPROFILER_EVENT_SCOPE(Interchange_SceneImportAsset_ContextMenuReset)
 	FInterchangeResetParameters ResetParameters(SceneImportAsset);
 
@@ -421,6 +446,11 @@ void UInterchangeEditorScriptLibrary::ResetActors(TArray<AActor*> Actors)
 {
 	using namespace UE::InterchangeEditorScriptLibrary::Private::InterchangeReset;
 
+	if (!IsInterchangeResetAvailable())
+	{
+		return;
+	}
+
 	TRACE_CPUPROFILER_EVENT_SCOPE(InterchangeEditorScriptLibrary::ResetActors)
 	TMap<UInterchangeSceneImportAsset*, FInterchangeResetParameters> BatchedReset;
 	TMap<UInterchangeSceneImportAsset*, ALevelInstance*> LevelInstanceResets;
@@ -476,6 +506,12 @@ void UInterchangeEditorScriptLibrary::ResetActors(TArray<AActor*> Actors)
 
 bool UInterchangeEditorScriptLibrary::CanResetActor(const AActor* Actor)
 {
+	using namespace UE::InterchangeEditorScriptLibrary::Private::InterchangeReset;
+	if (!IsInterchangeResetAvailable())
+	{
+		return false;
+	}
+
 	if (const ALevelInstance* LevelInstanceActor = Cast<ALevelInstance>(Actor))
 	{
 		const TSoftObjectPtr<UWorld>& WorldObjectPtr = LevelInstanceActor->GetWorldAsset();
@@ -504,6 +540,12 @@ bool UInterchangeEditorScriptLibrary::CanResetActor(const AActor* Actor)
 
 bool UInterchangeEditorScriptLibrary::CanResetWorld(const UWorld* World)
 {
+	using namespace UE::InterchangeEditorScriptLibrary::Private::InterchangeReset;
+	if (!IsInterchangeResetAvailable())
+	{
+		return false;
+	}
+
 	if (!World)
 	{
 		return false;
diff --git a/Engine/Plugins/Interchange/Editor/Source/InterchangeEditor/Private/InterchangeResetContextMenuExtender.cpp b/Engine/Plugins/Interchange/Editor/Source/InterchangeEditor/Private/InterchangeResetContextMenuExtender.cpp
--- a/Engine/Plugins/Interchange/Editor/Source/InterchangeEditor/Private/InterchangeResetContextMenuExtender.cpp
+++ b/Engine/Plugins/Interchange/Editor/Source/InterchangeEditor/Private/InterchangeResetContextMenuExtender.cpp
@@ -11,8 +11,10 @@
 #include "Engine/Selection.h"
 #include "EngineUtils.h"
 #include "Framework/MultiBox/MultiBoxBuilder.h"
+#include "HAL/IConsoleManager.h"
 #include "InterchangeAssetUserData.h"
 #include "InterchangeEditorScriptLibrary.h"
+#include "InterchangeImportReset.h"
 #include "InterchangeLevelFactoryNode.h"
 #include "InterchangeManager.h"
 #include "InterchangeSceneImportAsset.h"
@@ -95,6 +97,13 @@ namespace UE::Interchange::InterchangeReset
 static FDelayedAutoRegisterHelper DelayedAutoRegisterLevelContextMenu(EDelayedRegisterRunPhase::EndOfEngineInit, [] {
 	UToolMenus::RegisterStartupCallback(FSimpleMulticastDelegate::FDelegate::CreateLambda([]()
 		{
+			IConsoleVariable* CVar = IConsoleManager::Get().FindConsoleVariable(TEXT("Interchange.FeatureFlags.InterchangeReset"), false);
+			const bool bInterchangeResetEnabled = CVar ? CVar->GetBool() : false;
+			if (!bInterchangeResetEnabled)
+			{
+				return;
+			}
+
 			FToolMenuOwnerScoped OwnerScoped(UE_MODULE_NAME);
 			UToolMenu* Menu = UE::ContentBrowser::ExtendToolMenu_AssetContextMenu(UWorld::StaticClass());
 
@@ -150,11 +159,21 @@ FDelegateHandle FInterchangeResetContextMenuExtender::LevelEditorExtenderDelegat
 
 void FInterchangeResetContextMenuExtender::SetupLevelEditorContextMenuExtender()
 {
-	FLevelEditorModule& LevelEditorModule = FModuleManager::LoadModuleChecked< FLevelEditorModule >("LevelEditor");
-	TArray< FLevelEditorModule::FLevelViewportMenuExtender_SelectedActors >& CBMenuExtenderDelegates = LevelEditorModule.GetAllLevelViewportContextMenuExtenders();
+	IConsoleVariable* CVar = IConsoleManager::Get().FindConsoleVariable(TEXT("Interchange.FeatureFlags.InterchangeReset"), false);
+	const bool bInterchangeResetEnabled = CVar ? CVar->GetBool() : false;
+	if (!bInterchangeResetEnabled)
+	{
+		return;
+	}
 
-	CBMenuExtenderDelegates.Add(FLevelEditorModule::FLevelViewportMenuExtender_SelectedActors::CreateStatic(&UE::Interchange::InterchangeReset::OnExtendLevelEditorActorSelectionMenu));
-	LevelEditorExtenderDelegateHandle = CBMenuExtenderDelegates.Last().GetHandle();
+	if (!LevelEditorExtenderDelegateHandle.IsValid())
+	{
+		FLevelEditorModule& LevelEditorModule = FModuleManager::LoadModuleChecked< FLevelEditorModule >("LevelEditor");
+		TArray< FLevelEditorModule::FLevelViewportMenuExtender_SelectedActors >& CBMenuExtenderDelegates = LevelEditorModule.GetAllLevelViewportContextMenuExtenders();
+
+		CBMenuExtenderDelegates.Add(FLevelEditorModule::FLevelViewportMenuExtender_SelectedActors::CreateStatic(&UE::Interchange::InterchangeReset::OnExtendLevelEditorActorSelectionMenu));
+		LevelEditorExtenderDelegateHandle = CBMenuExtenderDelegates.Last().GetHandle();
+	}
 }
 
 void FInterchangeResetContextMenuExtender::RemoveLevelEditorContextMenuExtender()
@@ -169,6 +188,7 @@ void FInterchangeResetContextMenuExtender::RemoveLevelEditorContextMenuExtender(
 				return Delegate.GetHandle() == LevelEditorExtenderDelegateHandle;
 			}
 		);
+		LevelEditorExtenderDelegateHandle.Reset();
 	}
 }
 
diff --git a/Engine/Source/Runtime/Interchange/Engine/Private/InterchangeImportReset.cpp b/Engine/Source/Runtime/Interchange/Engine/Private/InterchangeImportReset.cpp
--- a/Engine/Source/Runtime/Interchange/Engine/Private/InterchangeImportReset.cpp
+++ b/Engine/Source/Runtime/Interchange/Engine/Private/InterchangeImportReset.cpp
@@ -20,6 +20,13 @@ namespace UE::Interchange::InterchangeReset
 		const FString FactoryNodeUidPathKey("InterchangeFactoryNodeUid");
 	}
 
+	bool GInterchangeResetEnabled = false;
+	FAutoConsoleVariableRef CCvarInterchangeResetEnabled(
+		TEXT("Interchange.FeatureFlags.InterchangeReset"),
+		GInterchangeResetEnabled,
+		TEXT("Whether Interchange Reset is available"),
+		ECVF_Default);
+
 	bool GInterchangeResetFilteredNodes = true;
 	FAutoConsoleVariableRef CCvarInterchangeResetFilteredNodes(
 		TEXT("Interchange.FeatureFlags.Reset.UseFilteredNodes"),
