[BUGFIX] Interchange UI - Import Dialog 'Show Sections' is not persisted between imports
- Setting the view identifier for details panel will persiste the details view config between sessions.

#rnx

diff --git a/Engine/Plugins/Interchange/Editor/Source/Pipelines/Private/SInterchangePipelineConfigurationDialog.cpp b/Engine/Plugins/Interchange/Editor/Source/Pipelines/Private/SInterchangePipelineConfigurationDialog.cpp
--- a/Engine/Plugins/Interchange/Editor/Source/Pipelines/Private/SInterchangePipelineConfigurationDialog.cpp
+++ b/Engine/Plugins/Interchange/Editor/Source/Pipelines/Private/SInterchangePipelineConfigurationDialog.cpp
@@ -651,6 +651,7 @@ TSharedRef<SBox> SInterchangePipelineConfigurationDialog::SpawnPipelineConfigura
 	DetailsViewArgs.bShowAnimatedPropertiesOption = false;
 	DetailsViewArgs.bShowHiddenPropertiesWhilePlayingOption = false;
 	DetailsViewArgs.NameAreaSettings = FDetailsViewArgs::HideNameArea;
+	DetailsViewArgs.ViewIdentifier = FName("InterchangeImportDialog");
 	PipelineConfigurationDetailsView = PropertyEditorModule.CreateDetailView(DetailsViewArgs);
 	InspectorBox->SetContent(PipelineConfigurationDetailsView->AsShared());
 	SetEditPipeline(nullptr);
@@ -1102,7 +1103,7 @@ void SInterchangePipelineConfigurationDialog::Construct(const FArguments& InArgs
 		FString KeyName = CurrentStackName.ToString() + TEXT("_LastSelectedPipeline");
 		if (GConfig->GetString(TEXT("InterchangeSelectPipeline"), *KeyName, LastPipelineName, GEditorPerProjectIni))
 		{
-			for (TSharedPtr<FInterchangePipelineItemType> PipelineItem: PipelineListViewItems)
+			for (TSharedPtr<FInterchangePipelineItemType> PipelineItem : PipelineListViewItems)
 			{
 				FString PipelineItemName = PipelineItem->Pipeline->GetClass()->GetName();
 				if (PipelineItemName.Equals(LastPipelineName))
