[BUGFIX] Interchange Level Import - Level Instance and Pakced Level Actors are causing an Asset Check
- FSoftObjectPath is adding a dependency between the level asset and scene import asset.
- Changing the data type to FString resolves the content validation error.

#rnx

diff --git a/Engine/Plugins/Interchange/Editor/Source/InterchangeEditor/Private/InterchangeEditorScriptLibrary.cpp b/Engine/Plugins/Interchange/Editor/Source/InterchangeEditor/Private/InterchangeEditorScriptLibrary.cpp
--- a/Engine/Plugins/Interchange/Editor/Source/InterchangeEditor/Private/InterchangeEditorScriptLibrary.cpp
+++ b/Engine/Plugins/Interchange/Editor/Source/InterchangeEditor/Private/InterchangeEditorScriptLibrary.cpp
@@ -109,7 +109,7 @@ namespace UE::InterchangeEditorScriptLibrary::Private
 				{
 					if (!AssetUserData->SceneImportPaths.IsEmpty())
 					{
-						return GetSceneImportAssetFromPath(AssetUserData->SceneImportPaths[0]);
+						return GetSceneImportAssetFromPath(FSoftObjectPath(AssetUserData->SceneImportPaths[0]));
 					}
 				}
 			}
@@ -143,7 +143,7 @@ namespace UE::InterchangeEditorScriptLibrary::Private
 						{
 							if (const FString* FactoryNodeUid = AssetUserData->MetaData.Find(Constants::FactoryNodeUidPathKey))
 							{
-								if (UInterchangeSceneImportAsset* SceneImportAsset = GetSceneImportAssetFromPath(*SceneImportAssetPathString))
+								if (UInterchangeSceneImportAsset* SceneImportAsset = GetSceneImportAssetFromPath(FSoftObjectPath(*SceneImportAssetPathString)))
 								{
 									if (SceneImportAsset->AssetImportData->GetStoredFactoryNode(*FactoryNodeUid) != nullptr)
 									{
@@ -276,10 +276,10 @@ namespace UE::InterchangeEditorScriptLibrary::Private
 		void ExecuteWorldReset(UWorld* World, const UInterchangeSceneImportAsset* SceneImportAsset)
 		{
 			TRACE_CPUPROFILER_EVENT_SCOPE(Interchange_Level_ContextMenuReset)
-				if (!World || !SceneImportAsset)
-				{
-					return;
-				}
+			if (!World || !SceneImportAsset)
+			{
+				return;
+			}
 
 			FWorldContext& EditorWorldContext = GEditor->GetEditorWorldContext();
 			if (UWorld* EditorWorld = EditorWorldContext.World())
@@ -365,9 +365,9 @@ void UInterchangeEditorScriptLibrary::ResetSceneImportAsset(UInterchangeSceneImp
 								using namespace UE::InterchangeEditorScriptLibrary::Private::InterchangeReset;
 
 								bool bShouldCreateLevelInstance = false;
-								for (const FSoftObjectPath& SceneImportAssetPath : AssetUserData->SceneImportPaths)
+								for (const FString& SceneImportAssetPath : AssetUserData->SceneImportPaths)
 								{
-									if (SceneImportAssetPath == CurrentSceneImportObjectPath)
+									if (FSoftObjectPath(SceneImportAssetPath) == CurrentSceneImportObjectPath)
 									{
 										bShouldCreateLevelInstance = true;
 										break;
diff --git a/Engine/Plugins/Interchange/Runtime/Source/Import/Private/Scene/InterchangeActorHelper.cpp b/Engine/Plugins/Interchange/Runtime/Source/Import/Private/Scene/InterchangeActorHelper.cpp
--- a/Engine/Plugins/Interchange/Runtime/Source/Import/Private/Scene/InterchangeActorHelper.cpp
+++ b/Engine/Plugins/Interchange/Runtime/Source/Import/Private/Scene/InterchangeActorHelper.cpp
@@ -258,11 +258,11 @@ namespace UE::Interchange::ActorHelper
 			return;
 		}
 
-		const FSoftObjectPath SceneImportAssetPath((const UObject*)SceneImportAsset);
 		if (AWorldSettings* WorldSettings = World->GetWorldSettings())
 		{
 			UInterchangeLevelAssetUserData* AssetUserData = NewObject<UInterchangeLevelAssetUserData>(World);
-			AssetUserData->SceneImportPaths.Add(SceneImportAssetPath);
+			const FSoftObjectPath SceneImportAssetPath((const UObject*)SceneImportAsset);
+			AssetUserData->SceneImportPaths.Add(SceneImportAssetPath.ToString());
 			WorldSettings->AddAssetUserData(AssetUserData);
 		}
 	}
diff --git a/Engine/Plugins/Interchange/Runtime/Source/Import/Public/InterchangeAssetUserData.h b/Engine/Plugins/Interchange/Runtime/Source/Import/Public/InterchangeAssetUserData.h
--- a/Engine/Plugins/Interchange/Runtime/Source/Import/Public/InterchangeAssetUserData.h
+++ b/Engine/Plugins/Interchange/Runtime/Source/Import/Public/InterchangeAssetUserData.h
@@ -26,5 +26,5 @@ class UInterchangeLevelAssetUserData : public UAssetUserData
 
 public:
 	UPROPERTY(VisibleAnywhere, BlueprintReadWrite, AdvancedDisplay, Category = "Interchange User Data", meta = (ScriptName = "SceneImportPaths", DisplayName = "Scene Import Paths"))
-	TArray<FSoftObjectPath> SceneImportPaths;
+	TArray<FString> SceneImportPaths;
 };
