[BUGFIX] Interchange Import Dialog - Reset Section Selection to 'All' when the import dialog spawns

- With a new addition to the Details view code, it is possible to optionally reset the sections to default when the details view is first opened.

#rnx

diff --git a/Engine/Plugins/Interchange/Editor/Source/Pipelines/Private/SInterchangePipelineConfigurationDialog.cpp b/Engine/Plugins/Interchange/Editor/Source/Pipelines/Private/SInterchangePipelineConfigurationDialog.cpp
--- a/Engine/Plugins/Interchange/Editor/Source/Pipelines/Private/SInterchangePipelineConfigurationDialog.cpp
+++ b/Engine/Plugins/Interchange/Editor/Source/Pipelines/Private/SInterchangePipelineConfigurationDialog.cpp
@@ -652,7 +652,21 @@ TSharedRef<SBox> SInterchangePipelineConfigurationDialog::SpawnPipelineConfigura
 	DetailsViewArgs.bShowHiddenPropertiesWhilePlayingOption = false;
 	DetailsViewArgs.NameAreaSettings = FDetailsViewArgs::HideNameArea;
 	DetailsViewArgs.ViewIdentifier = FName("InterchangeImportDialog");
+
 	PipelineConfigurationDetailsView = PropertyEditorModule.CreateDetailView(DetailsViewArgs);
+
+	// Reset to 'All' Sections on spawning the modal only. Subsequent changes to details view objects will use the persisted section selection.
+	bPendingDetailsViewSectionReset = true;
+	PipelineConfigurationDetailsView->SetOnObjectArrayChanged(FOnObjectArrayChanged::CreateLambda(
+		[this](const FString& /*Title*/, const TArray<UObject*>& /*SelectedObjects*/)
+		{
+			if (bPendingDetailsViewSectionReset)
+			{
+				bPendingDetailsViewSectionReset = !PipelineConfigurationDetailsView->ResetToDefaultSection();
+			}
+		})
+	);
+
 	InspectorBox->SetContent(PipelineConfigurationDetailsView->AsShared());
 	SetEditPipeline(nullptr);
 	PipelineConfigurationDetailsView->GetIsPropertyVisibleDelegate().BindLambda([this](const FPropertyAndParent& PropertyAndParent)
diff --git a/Engine/Plugins/Interchange/Editor/Source/Pipelines/Private/SInterchangePipelineConfigurationDialog.h b/Engine/Plugins/Interchange/Editor/Source/Pipelines/Private/SInterchangePipelineConfigurationDialog.h
--- a/Engine/Plugins/Interchange/Editor/Source/Pipelines/Private/SInterchangePipelineConfigurationDialog.h
+++ b/Engine/Plugins/Interchange/Editor/Source/Pipelines/Private/SInterchangePipelineConfigurationDialog.h
@@ -201,6 +201,8 @@ private:
 	bool bCanceled = false;
 	bool bImportAll = false;
 
+	bool bPendingDetailsViewSectionReset = false;
+
 	// This is used to denote that the import dialog is opened from the Interchange Test Plan asset to configure the chosen Pipelines.
 	bool bTestConfigDialog = false;
 
