[Interchange] Legacy FBX to Interchange AssetUserData saving fix.

#rnx

diff --git a/Engine/Plugins/Interchange/Editor/Source/InterchangeEditor/Private/InterchangeFbxAssetImportDataConverter.cpp b/Engine/Plugins/Interchange/Editor/Source/InterchangeEditor/Private/InterchangeFbxAssetImportDataConverter.cpp
--- a/Engine/Plugins/Interchange/Editor/Source/InterchangeEditor/Private/InterchangeFbxAssetImportDataConverter.cpp
+++ b/Engine/Plugins/Interchange/Editor/Source/InterchangeEditor/Private/InterchangeFbxAssetImportDataConverter.cpp
@@ -1081,7 +1081,7 @@ bool UInterchangeFbxAssetImportDataConverter::ConvertImportData(const UObject* S
 	UInterchangeManager& InterchangeManager = UInterchangeManager::GetInterchangeManager();
 	if (SourceImportData->IsA<UInterchangeAssetImportData>())
 	{
-		if (UFbxImportUI* FbxImportUI = UE::Interchange::Private::ConvertToLegacyFbx(GetTransientPackage(), Cast<UInterchangeAssetImportData>(SourceImportData)))
+		if (UFbxImportUI* FbxImportUI = UE::Interchange::Private::ConvertToLegacyFbx(SourceImportData->GetOuter(), Cast<UInterchangeAssetImportData>(SourceImportData)))
 		{
 			bResult = true;
 			if (DestinationClass->IsChildOf<UFbxImportUI>())
@@ -1128,7 +1128,7 @@ bool UInterchangeFbxAssetImportDataConverter::ConvertImportData(const UObject* S
 					DestinationAssetImportData->ImportUniformScale = SourceAssetImportData->ImportUniformScale;
 				};
 
-			UFbxImportUI* TempFbxImportUI = NewObject<UFbxImportUI>(GetTransientPackage());
+			UFbxImportUI* TempFbxImportUI = NewObject<UFbxImportUI>(SourceImportData->GetOuter());
 			TempFbxImportUI->bImportMaterials = false;
 			TempFbxImportUI->bImportAsSkeletal = false;
 			TempFbxImportUI->bImportMesh = false;
@@ -1183,7 +1183,7 @@ bool UInterchangeFbxAssetImportDataConverter::ConvertImportData(const UObject* S
 		if (FbxImportUI)
 		{
 			//Convert Legacy Fbx to Interchange
-			*DestinationImportData = UE::Interchange::Private::ConvertToInterchange(GetTransientPackage(), FbxImportUI);
+			*DestinationImportData = UE::Interchange::Private::ConvertToInterchange(SourceImportData->GetOuter(), FbxImportUI);
 			bResult = true;
 		}
 	}
