[Interchange] Improve Logging Capabilities

#rnx

diff --git a/Engine/Plugins/Interchange/Editor/Source/InterchangeEditor/Private/InterchangeEditorModule.cpp b/Engine/Plugins/Interchange/Editor/Source/InterchangeEditor/Private/InterchangeEditorModule.cpp
--- a/Engine/Plugins/Interchange/Editor/Source/InterchangeEditor/Private/InterchangeEditorModule.cpp
+++ b/Engine/Plugins/Interchange/Editor/Source/InterchangeEditor/Private/InterchangeEditorModule.cpp
@@ -67,7 +67,9 @@ namespace UE::Interchange::InterchangeEditorModule
 				}
 			}
 
-			LogListing->AddMessages(TokenizedMessages);
+			// UInterchangeManager::ReleaseAsyncHelper has output logs turned on. To avoid duplication, turning it off here.
+			constexpr bool bMirrorToOutputLog = false;
+			LogListing->AddMessages(TokenizedMessages, bMirrorToOutputLog);
 			LogListing->NotifyIfAnyMessages(NSLOCTEXT("Interchange", "LogAndNotify", "There were issues with the import."), EMessageSeverity::Info);
 		}
 	}
diff --git a/Engine/Source/Runtime/Interchange/Engine/Private/InterchangeManager.cpp b/Engine/Source/Runtime/Interchange/Engine/Private/InterchangeManager.cpp
--- a/Engine/Source/Runtime/Interchange/Engine/Private/InterchangeManager.cpp
+++ b/Engine/Source/Runtime/Interchange/Engine/Private/InterchangeManager.cpp
@@ -2740,7 +2740,7 @@ void UInterchangeManager::ReleaseAsyncHelper(TWeakPtr<UE::Interchange::FImportAs
 
 	check(AsyncHelper.IsValid());
 	
-	constexpr bool bLogWarningsAndErrors = false;
+	constexpr bool bLogWarningsAndErrors = true;
 	
 	double ImportTimeDelta;
 	bool bSucceeded = false;
