[Interchange UI]: Interchange Cards View Improvements and the Details View Improvements
- Label changes in the Import Dialog
- Changed the Pipeline Stack Selection Text Combobox into a full combobox with icon and text
- Migrated the 'Show Essentials' and 'Filter On Contents' as a Detail View options menu extension
- Made the cards view show disabled cards instead of hide cards.

#rnx
#changelist validated

diff --git a/Engine/Plugins/Interchange/Editor/Source/Pipelines/Private/InterchangeEditorPipelineStyle.h b/Engine/Plugins/Interchange/Editor/Source/Pipelines/Private/InterchangeEditorPipelineStyle.h
--- a/Engine/Plugins/Interchange/Editor/Source/Pipelines/Private/InterchangeEditorPipelineStyle.h
+++ b/Engine/Plugins/Interchange/Editor/Source/Pipelines/Private/InterchangeEditorPipelineStyle.h
@@ -16,7 +16,6 @@ class FInterchangeEditorPipelineStyle final
 	: public FSlateStyleSet
 {
 public:
-
 	/** Default constructor. */
 	FInterchangeEditorPipelineStyle()
 		: FSlateStyleSet("InterchangeEditorPipelineStyle")
@@ -44,6 +43,8 @@ public:
 		Set("PipelineConfigurationIcon.PipelineStack", PipelineStackBrush16);
 		FSlateImageBrush* PipelineStackDefaultBrush16 = new FSlateImageBrush(RootToContentDir(TEXT("Resources/Interchange_PipelineStackDefault_Icon_16"), TEXT(".png")), Icon16x16);
 		Set("PipelineConfigurationIcon.PipelineStackDefault", PipelineStackDefaultBrush16);
+		FSlateImageBrush* PipelineStackPresetBrush20 = new FSlateImageBrush(RootToContentDir(TEXT("Resources/Interchange_PipelineStackPreset_Icon_20"), TEXT(".png")), Icon20x20);
+		Set("PipelineConfigurationIcon.PipelineStackPreset", PipelineStackPresetBrush20);
 		FSlateImageBrush* TranslatorSettingsBrush16 = new FSlateImageBrush(RootToContentDir(TEXT("Resources/Interchange_TranslatorSettings_16"), TEXT(".png")), Icon16x16);
 		Set("PipelineConfigurationIcon.TranslatorSettings", TranslatorSettingsBrush16);
 
@@ -55,7 +56,6 @@ public:
 		Set("AssetCard.Header.Border", new FSlateRoundedBoxBrush(FStyleColors::Dropdown.GetSpecifiedColor(), 4.f));
 		Set("AssetCard.Body.Border", new FSlateRoundedBoxBrush(FStyleColors::Panel.GetSpecifiedColor(), 0.f));
 
-
 		const FButtonStyle ResetSelectedPipelineButtonStyle = FButtonStyle()
 			.SetNormal(FSlateRoundedBoxBrush(FStyleColors::Dropdown, 4.0f))
 			.SetHovered(FSlateRoundedBoxBrush(FStyleColors::Dropdown, 4.0f))
diff --git a/Engine/Plugins/Interchange/Editor/Source/Pipelines/Private/InterchangePipelineConfigurationGeneric.cpp b/Engine/Plugins/Interchange/Editor/Source/Pipelines/Private/InterchangePipelineConfigurationGeneric.cpp
--- a/Engine/Plugins/Interchange/Editor/Source/Pipelines/Private/InterchangePipelineConfigurationGeneric.cpp
+++ b/Engine/Plugins/Interchange/Editor/Source/Pipelines/Private/InterchangePipelineConfigurationGeneric.cpp
@@ -34,6 +34,7 @@ namespace UE::Interchange::Private
 		static FString SizeYStr = TEXT("SizeY");
 		return SizeYStr;
 	}
+
 	void RestoreClientSize(FVector2D& OutClientSize, const FString& WindowsUniqueID)
 	{
 		if (GConfig->DoesSectionExist(TEXT("InterchangeImportDialogOptions"), GEditorPerProjectIni))
@@ -44,6 +45,7 @@ namespace UE::Interchange::Private
 			GConfig->GetDouble(TEXT("InterchangeImportDialogOptions"), *SizeYUid, OutClientSize.Y, GEditorPerProjectIni);
 		}
 	}
+
 	void BackupClientSize(SWindow* Window, const FString& WindowsUniqueID)
 	{
 		if (!Window)
@@ -59,131 +61,24 @@ namespace UE::Interchange::Private
 		GConfig->SetDouble(TEXT("InterchangeImportDialogOptions"), *SizeXUid, ClientSize.X, GEditorPerProjectIni);
 		GConfig->SetDouble(TEXT("InterchangeImportDialogOptions"), *SizeYUid, ClientSize.Y, GEditorPerProjectIni);
 	}
-}
 
-EInterchangePipelineConfigurationDialogResult UInterchangePipelineConfigurationGeneric::ShowPipelineConfigurationDialog(TArray<FInterchangeStackInfo>& PipelineStacks
-	, TArray<UInterchangePipelineBase*>& OutPipelines
-	, TWeakObjectPtr<UInterchangeSourceData> SourceData
-	, TWeakObjectPtr <UInterchangeTranslatorBase> Translator
-	, TWeakObjectPtr<UInterchangeBaseNodeContainer> BaseNodeContainer)
-{
-	using namespace UE::Interchange::Private;
-	//Create and show the graph inspector UI dialog
-	TSharedPtr<SWindow> ParentWindow;
-	if(IMainFrameModule* MainFrame = FModuleManager::LoadModulePtr<IMainFrameModule>("MainFrame"))
+	FString GetWindowsUID(bool bSceneImport, bool bInvokedThroughTestPlan)
 	{
-		ParentWindow = MainFrame->GetParentWindow();
-	}
-	FString WindowDialogUid = TEXT("ImportContentDialog");
-	FVector2D ClientSize = DefaultClientSize();
-	RestoreClientSize(ClientSize, WindowDialogUid);
-
-	TSharedRef<SWindow> Window = SNew(SWindow)
-		.ClientSize(ClientSize)
-		.MinWidth(GetMinSizeX())
-		.MinHeight(GetMinSizeY())
-		.Title(NSLOCTEXT("Interchange", "PipelineConfigurationGenericTitleContent", "Import Content"));
-	TSharedPtr<SInterchangePipelineConfigurationDialog> InterchangePipelineConfigurationDialog;
-
-	Window->SetContent
-	(
-		SAssignNew(InterchangePipelineConfigurationDialog, SInterchangePipelineConfigurationDialog)
-		.OwnerWindow(Window)
-		.SourceData(SourceData)
-		.bSceneImport(false)
-		.bReimport(false)
-		.bTestConfigDialog(false)
-		.PipelineStacks(PipelineStacks)
-		.OutPipelines(&OutPipelines)
-		.BaseNodeContainer(BaseNodeContainer)
-		.Translator(Translator)
-	);
-
-	Window->SetOnWindowClosed(FOnWindowClosed::CreateLambda([&WindowDialogUid](const TSharedRef<SWindow>& WindowMoved)
+		if (bInvokedThroughTestPlan)
 		{
-			BackupClientSize(&WindowMoved.Get(), WindowDialogUid);
-		}));
-
-	FSlateApplication::Get().AddModalWindow(Window, ParentWindow, false);
-	
-	if (InterchangePipelineConfigurationDialog->IsCanceled())
-	{
-		return EInterchangePipelineConfigurationDialogResult::Cancel;
-	}
-
-	if (InterchangePipelineConfigurationDialog->IsImportAll())
-	{
-		return EInterchangePipelineConfigurationDialogResult::ImportAll;
-	}
-
-	return EInterchangePipelineConfigurationDialogResult::Import;
-}
-
-EInterchangePipelineConfigurationDialogResult UInterchangePipelineConfigurationGeneric::ShowScenePipelineConfigurationDialog(TArray<FInterchangeStackInfo>& PipelineStacks
-	, TArray<UInterchangePipelineBase*>& OutPipelines
-	, TWeakObjectPtr<UInterchangeSourceData> SourceData
-	, TWeakObjectPtr <UInterchangeTranslatorBase> Translator
-	, TWeakObjectPtr<UInterchangeBaseNodeContainer> BaseNodeContainer)
-{
-	using namespace UE::Interchange::Private;
-	//Create and show the graph inspector UI dialog
-	TSharedPtr<SWindow> ParentWindow;
-	if (IMainFrameModule* MainFrame = FModuleManager::LoadModulePtr<IMainFrameModule>("MainFrame"))
-	{
-		ParentWindow = MainFrame->GetParentWindow();
-	}
-
-	FString WindowDialogUid = TEXT("ImportSceneDialog");
-	FVector2D ClientSize = DefaultClientSize();
-	RestoreClientSize(ClientSize, WindowDialogUid);
-	TSharedRef<SWindow> Window = SNew(SWindow)
-		.ClientSize(ClientSize)
-		.MinWidth(GetMinSizeX())
-		.MinHeight(GetMinSizeY())
-		.Title(NSLOCTEXT("Interchange", "PipelineConfigurationGenericTitleScene", "Import Scene"));
-	TSharedPtr<SInterchangePipelineConfigurationDialog> InterchangePipelineConfigurationDialog;
-
-	Window->SetContent
-	(
-		SAssignNew(InterchangePipelineConfigurationDialog, SInterchangePipelineConfigurationDialog)
-		.OwnerWindow(Window)
-		.SourceData(SourceData)
-		.bSceneImport(true)
-		.bReimport(false)
-		.bTestConfigDialog(false)
-		.PipelineStacks(PipelineStacks)
-		.OutPipelines(&OutPipelines)
-		.BaseNodeContainer(BaseNodeContainer)
-		.Translator(Translator)
-	);
+			return TEXT("TestPlanConfigurationDialog");
+		}
 
-	Window->SetOnWindowClosed(FOnWindowClosed::CreateLambda([&WindowDialogUid](const TSharedRef<SWindow>& WindowMoved)
+		if (bSceneImport)
 		{
-			BackupClientSize(&WindowMoved.Get(), WindowDialogUid);
-		}));
-
-	FSlateApplication::Get().AddModalWindow(Window, ParentWindow, false);
+			return TEXT("ImportSceneDialog");
+		}
 
-	if (InterchangePipelineConfigurationDialog->IsCanceled())
-	{
-		return EInterchangePipelineConfigurationDialogResult::Cancel;
+		return TEXT("ImportContentDialog");
 	}
-
-	if (InterchangePipelineConfigurationDialog->IsImportAll())
-	{
-		return EInterchangePipelineConfigurationDialogResult::ImportAll;
-	}
-
-	return EInterchangePipelineConfigurationDialogResult::Import;
 }
 
-EInterchangePipelineConfigurationDialogResult UInterchangePipelineConfigurationGeneric::ShowReimportPipelineConfigurationDialog(TArray<FInterchangeStackInfo>& PipelineStacks
-	, TArray<UInterchangePipelineBase*>& OutPipelines
-	, TWeakObjectPtr<UInterchangeSourceData> SourceData
-	, TWeakObjectPtr <UInterchangeTranslatorBase> Translator
-	, TWeakObjectPtr<UInterchangeBaseNodeContainer> BaseNodeContainer
-	, TWeakObjectPtr <UObject> ReimportAsset
-	, bool bSceneImport)
+EInterchangePipelineConfigurationDialogResult UInterchangePipelineConfigurationGeneric::ShowPipelineDialog_Internal(FPipelineConfigurationDialogParams& InParams)
 {
 	using namespace UE::Interchange::Private;
 	//Create and show the graph inspector UI dialog
@@ -193,7 +88,7 @@ EInterchangePipelineConfigurationDialogResult UInterchangePipelineConfigurationG
 		ParentWindow = MainFrame->GetParentWindow();
 	}
 
-	FString WindowDialogUid = TEXT("ImportContentDialog");
+	const FString WindowDialogUid = GetWindowsUID(InParams.bSceneImport, InParams.bInvokedThroughTestPlan);
 	FVector2D ClientSize = DefaultClientSize();
 	RestoreClientSize(ClientSize, WindowDialogUid);
 
@@ -201,23 +96,41 @@ EInterchangePipelineConfigurationDialogResult UInterchangePipelineConfigurationG
 		.ClientSize(ClientSize)
 		.MinWidth(GetMinSizeX())
 		.MinHeight(GetMinSizeY())
-		.Title(NSLOCTEXT("Interchange", "PipelineConfigurationGenericTitleReimportContent", "Reimport Content"));
-	TSharedPtr<SInterchangePipelineConfigurationDialog> InterchangePipelineConfigurationDialog;
-
-	Window->SetContent
-	(
-		SAssignNew(InterchangePipelineConfigurationDialog, SInterchangePipelineConfigurationDialog)
+		.Title_Lambda([bSceneImport = InParams.bSceneImport, bReimport = InParams.bReimport, bInvokedThroughTestPlan = InParams.bInvokedThroughTestPlan]() 
+		{
+			if (bInvokedThroughTestPlan)
+			{
+				return NSLOCTEXT("Interchange", "PipelineConfigurationGenericTitlePipelineConfiguration", "Pipeline Configuration");
+			}
+
+			if (bReimport)
+			{
+				return NSLOCTEXT("Interchange", "PipelineConfigurationGenericTitleReimportContent", "Reimport Content");
+			}
+
+			if (bSceneImport)
+			{
+				return NSLOCTEXT("Interchange", "PipelineConfigurationGenericTitleScene", "Import Scene");
+			}
+			
+			return NSLOCTEXT("Interchange", "PipelineConfigurationGenericTitleContent", "Import Content");
+		});
+
+	TSharedPtr<SInterchangePipelineConfigurationDialog> InterchangePipelineConfigurationDialog = SNew(SInterchangePipelineConfigurationDialog)
 		.OwnerWindow(Window)
-		.SourceData(SourceData)
-		.bSceneImport(bSceneImport)
-		.bReimport(true)
-		.bTestConfigDialog(false)
-		.PipelineStacks(PipelineStacks)
-		.OutPipelines(&OutPipelines)
-		.BaseNodeContainer(BaseNodeContainer)
-		.ReimportObject(ReimportAsset)
-		.Translator(Translator)
-	);
+		.SourceData(InParams.SourceData)
+		.bSceneImport(InParams.bSceneImport)
+		.bReimport(InParams.bReimport)
+		.bTestConfigDialog(InParams.bInvokedThroughTestPlan)
+		.PipelineStacks(InParams.PipelineStacks)
+		.OutPipelines(&InParams.OutPipelines)
+		.BaseNodeContainer(InParams.BaseNodeContainer)
+		.ReimportObject(InParams.ReimportAsset)
+		.Translator(InParams.Translator)
+		.bOverrideDefaultShowEssentials(InParams.bOverrideDefaultShowEssentials)
+		.bOverrideDefaultFilterOnContent(InParams.bOverrideDefaultFilterOnContent);
+
+	Window->SetContent(InterchangePipelineConfigurationDialog.ToSharedRef());
 
 	Window->SetOnWindowClosed(FOnWindowClosed::CreateLambda([&WindowDialogUid](const TSharedRef<SWindow>& WindowMoved)
 		{
@@ -231,61 +144,17 @@ EInterchangePipelineConfigurationDialogResult UInterchangePipelineConfigurationG
 		return EInterchangePipelineConfigurationDialogResult::Cancel;
 	}
 
-	if (InterchangePipelineConfigurationDialog->IsImportAll())
+	if (InParams.bInvokedThroughTestPlan)
 	{
-		return EInterchangePipelineConfigurationDialogResult::ImportAll;
+		return EInterchangePipelineConfigurationDialogResult::SaveConfig;
 	}
-
-	return EInterchangePipelineConfigurationDialogResult::Import;
-}
-
-EInterchangePipelineConfigurationDialogResult UInterchangePipelineConfigurationGeneric::ShowTestPlanConfigurationDialog(TArray<FInterchangeStackInfo>& PipelineStacks, TArray<UInterchangePipelineBase*>& OutPipelines, TWeakObjectPtr<UInterchangeSourceData> SourceData, TWeakObjectPtr<UInterchangeTranslatorBase> Translator, TWeakObjectPtr<UInterchangeBaseNodeContainer> BaseNodeContainer, TWeakObjectPtr<UObject> ReimportAsset, bool bSceneImport, bool bReimport)
-{
-	using namespace UE::Interchange::Private;
-	//Create and show the graph inspector UI dialog
-	TSharedPtr<SWindow> ParentWindow;
-	if (IMainFrameModule* MainFrame = FModuleManager::LoadModulePtr<IMainFrameModule>("MainFrame"))
+	else
 	{
-		ParentWindow = MainFrame->GetParentWindow();
-	}
-
-	FString WindowDialogUid = TEXT("TestPlanConfigurationDialog");
-	FVector2D ClientSize = DefaultClientSize();
-	RestoreClientSize(ClientSize, WindowDialogUid);
-
-	TSharedRef<SWindow> Window = SNew(SWindow)
-		.ClientSize(ClientSize)
-		.MinWidth(GetMinSizeX())
-		.MinHeight(GetMinSizeY())
-		.Title(NSLOCTEXT("Interchange", "PipelineConfigurationGenericTitlePipelineConfiguration", "Pipeline Configuration"));
-	TSharedPtr<SInterchangePipelineConfigurationDialog> InterchangePipelineConfigurationDialog;
-
-	Window->SetContent
-	(
-		SAssignNew(InterchangePipelineConfigurationDialog, SInterchangePipelineConfigurationDialog)
-		.OwnerWindow(Window)
-		.SourceData(SourceData)
-		.bSceneImport(bSceneImport)
-		.bReimport(bReimport)
-		.bTestConfigDialog(true)
-		.PipelineStacks(PipelineStacks)
-		.OutPipelines(&OutPipelines)
-		.BaseNodeContainer(BaseNodeContainer)
-		.ReimportObject(ReimportAsset)
-		.Translator(Translator)
-	);
-
-	Window->SetOnWindowClosed(FOnWindowClosed::CreateLambda([&WindowDialogUid](const TSharedRef<SWindow>& WindowMoved)
+		if (InterchangePipelineConfigurationDialog->IsImportAll())
 		{
-			BackupClientSize(&WindowMoved.Get(), WindowDialogUid);
-		}));
-
-	FSlateApplication::Get().AddModalWindow(Window, ParentWindow, false);
+			return EInterchangePipelineConfigurationDialogResult::ImportAll;
+		}
 
-	if (InterchangePipelineConfigurationDialog->IsCanceled())
-	{
-		return EInterchangePipelineConfigurationDialogResult::Cancel;
+		return EInterchangePipelineConfigurationDialogResult::Import;
 	}
-
-	return EInterchangePipelineConfigurationDialogResult::SaveConfig;
-}
+}
\ No newline at end of file
diff --git a/Engine/Plugins/Interchange/Editor/Source/Pipelines/Private/SInterchangeAssetCard.cpp b/Engine/Plugins/Interchange/Editor/Source/Pipelines/Private/SInterchangeAssetCard.cpp
--- a/Engine/Plugins/Interchange/Editor/Source/Pipelines/Private/SInterchangeAssetCard.cpp
+++ b/Engine/Plugins/Interchange/Editor/Source/Pipelines/Private/SInterchangeAssetCard.cpp
@@ -38,6 +38,11 @@ SInterchangeAssetCard::~SInterchangeAssetCard()
 
 void SInterchangeAssetCard::RefreshCard(UInterchangeBaseNodeContainer* InPreviewNodeContainer)
 {
+	if (bCardDisabled)
+	{
+		return;
+	}
+
 	check(InPreviewNodeContainer);
 	CardAssetCount = 0;
 	CardAssetToImportCount = 0;
@@ -87,6 +92,11 @@ void SInterchangeAssetCard::RefreshCard(UInterchangeBaseNodeContainer* InPreview
 
 bool SInterchangeAssetCard::RefreshHasConflicts(const TArray<FInterchangeConflictInfo>& InConflictInfos)
 {
+	if (bCardDisabled)
+	{
+		return false;
+	}
+
 	for (const FInterchangeConflictInfo& ConflictInfo : InConflictInfos)
 	{
 		if (ConflictInfo.AffectedAssetClasses.Contains(AssetClass))
@@ -110,6 +120,7 @@ void SInterchangeAssetCard::Construct(const FArguments& InArgs)
 	check(ShouldImportAssetType.IsBound());
 	OnImportAssetTypeChanged = InArgs._OnImportAssetTypeChanged;
 	check(OnImportAssetTypeChanged.IsBound());
+	bCardDisabled = InArgs._CardDisabled;
 
 	RefreshCard(InArgs._PreviewNodeContainer);
 
@@ -141,16 +152,13 @@ void SInterchangeAssetCard::Construct(const FArguments& InArgs)
 	ChildSlot
 	[
 		SNew(SVerticalBox)
-		.Visibility_Lambda([this]()
-			{
-				return CardAssetToImportCount > 0 ? EVisibility::Visible : EVisibility::Collapsed;
-			})
 		+ SVerticalBox::Slot()
 		.AutoHeight()
 		.Padding(0.0f, 4.0f, 0.0f, 0.0f)
 		[
 			SNew(SBorder)
 			.BorderImage(HeaderBorderBrush)
+			.IsEnabled(!bCardDisabled && (CardAssetToImportCount > 0))
 			[
 				SNew(SHorizontalBox)
 				+ SHorizontalBox::Slot()
@@ -171,11 +179,25 @@ void SInterchangeAssetCard::Construct(const FArguments& InArgs)
 					.Font(FAppStyle::GetFontStyle("NormalFont"))
 					.Text_Lambda([this]()
 						{
-							return FText::Format(LOCTEXT("CardLabelText", "{0} ({1} {2})"), FText::FromString(AssetClass->GetName()), FText::FromString(FString::FromInt(CardAssetToImportCount)), CardAssetToImportCount > 1 ? LOCTEXT("CardAssetsPlural", "assets") : LOCTEXT("CardAssetSingle", "asset"));
+							if (CardAssetToImportCount == 0)
+							{
+								return FText::Format(LOCTEXT("CardLabelText_NoAsset", "{0} (Asset type not found in source or disabled in this pipeline stack)"), FText::FromString(AssetClass->GetName()));
+							}
+							else
+							{
+								return FText::Format(LOCTEXT("CardLabelText_ValidAssets", "{0} ({1} {2})"), FText::FromString(AssetClass->GetName()), FText::FromString(FString::FromInt(CardAssetToImportCount)), CardAssetToImportCount > 1 ? LOCTEXT("CardAssetsPlural", "assets") : LOCTEXT("CardAssetSingle", "asset"));
+							}
 						})
 					.ToolTipText_Lambda([this]()
 						{
-							return FText::FromString(CardTooltip);
+							if (CardAssetToImportCount == 0)
+							{
+								return LOCTEXT("CardTooltip_NoAssets", "This might be because the source file doesn't have asset of this type or the import options are not set to import this type.");
+							}
+							else
+							{
+								return FText::FromString(CardTooltip);
+							}
 						})
 				]
 				+ SHorizontalBox::Slot()
@@ -219,6 +241,7 @@ void SInterchangeAssetCard::Construct(const FArguments& InArgs)
 		[
 			SNew(SBorder)
 			.BorderImage(BodyBorderBrush)
+			.IsEnabled(!bCardDisabled && (CardAssetToImportCount > 0))
 			[
 				SNew(SHorizontalBox)
 				+ SHorizontalBox::Slot()
@@ -252,7 +275,7 @@ void SInterchangeAssetCard::Construct(const FArguments& InArgs)
 						SNew(SCheckBox)
 						.IsChecked_Lambda([this]()
 						{
-							if (ShouldImportAssetType.Execute())
+							if (CardAssetToImportCount != 0 && ShouldImportAssetType.Execute())
 							{
 								return ECheckBoxState::Checked;
 							}
diff --git a/Engine/Plugins/Interchange/Editor/Source/Pipelines/Private/SInterchangeAssetCard.h b/Engine/Plugins/Interchange/Editor/Source/Pipelines/Private/SInterchangeAssetCard.h
--- a/Engine/Plugins/Interchange/Editor/Source/Pipelines/Private/SInterchangeAssetCard.h
+++ b/Engine/Plugins/Interchange/Editor/Source/Pipelines/Private/SInterchangeAssetCard.h
@@ -24,11 +24,13 @@ public:
 	SLATE_BEGIN_ARGS(SInterchangeAssetCard)
 		: _PreviewNodeContainer(nullptr)
 		, _AssetClass(nullptr)
+		, _CardDisabled(false)
 		, _ShouldImportAssetType()
 		, _OnImportAssetTypeChanged()
 		{}
 		SLATE_ARGUMENT(UInterchangeBaseNodeContainer*, PreviewNodeContainer)
 		SLATE_ARGUMENT(UClass*, AssetClass)
+		SLATE_ARGUMENT(bool, CardDisabled)
 		SLATE_EVENT(FInterchangeCard_ShouldImportAssetType, ShouldImportAssetType)
 		SLATE_EVENT(FInterchangeCard_OnImportAssetTypeChanged, OnImportAssetTypeChanged)
 	SLATE_END_ARGS()
@@ -48,6 +50,7 @@ protected:
 	int32 CardAssetDisabledCount = 0;
 	FString CardTooltip;
 	bool bHasConflictWarnings = false;
+	bool bCardDisabled = false;
 
 	/** The factory asset class so we can know which kind of Unreal asset this card is for. */
 	UClass* AssetClass = nullptr;
diff --git a/Engine/Plugins/Interchange/Editor/Source/Pipelines/Private/SInterchangePipelineConfigurationDialog.cpp b/Engine/Plugins/Interchange/Editor/Source/Pipelines/Private/SInterchangePipelineConfigurationDialog.cpp
--- a/Engine/Plugins/Interchange/Editor/Source/Pipelines/Private/SInterchangePipelineConfigurationDialog.cpp
+++ b/Engine/Plugins/Interchange/Editor/Source/Pipelines/Private/SInterchangePipelineConfigurationDialog.cpp
@@ -36,7 +36,6 @@
 #include "Widgets/Input/SCheckBox.h"
 #include "Widgets/Layout/SScrollBox.h"
 #include "Widgets/Layout/SSeparator.h"
-#include "Widgets/Input/STextComboBox.h"
 #include "Widgets/Layout/SUniformGridPanel.h"
 #include "Widgets/Input/SEditableTextBox.h"
 
@@ -74,6 +73,44 @@ constexpr double AdvancedUIRatio = 1.25;
 
 const FName ReimportStackName = TEXT("ReimportPipeline");
 
+// Pipeline Stack Preset Items
+void SInterchangePipelineStackPresetItem::Construct(const FArguments& InArgs)
+{
+	PresetItem = InArgs._PresetItem;
+
+	ensure(PresetItem.IsValid());
+	AddSlot()
+	.AutoWidth()
+	.VAlign(VAlign_Center)
+	.Padding(2.0, 0.0, 8.0, 0.0)
+	[
+		SNew(SBox)
+		.WidthOverride(14.0f)
+		.HeightOverride(14.0f)
+		[
+			SNew(SImage)
+			.Image_Static(&SInterchangePipelineStackPresetItem::GetItemIconBrush)
+		]
+	];
+
+	AddSlot()
+	.VAlign(VAlign_Center)
+	.FillWidth(1.0f)
+	[
+		SNew(STextBlock)
+		.Text(FText::FromString(*PresetItem))
+	];
+}
+
+const FSlateBrush* SInterchangePipelineStackPresetItem::GetItemIconBrush()
+{
+	static const FName IconName = FName("PipelineConfigurationIcon.PipelineStackPreset");
+	const FSlateIcon SlateIcon = FSlateIconFinder::FindIcon(IconName);
+	return SlateIcon.GetOptionalIcon();
+}
+
+// Pipeline List View Items
+
 void SInterchangePipelineItem::Construct(
 	const FArguments& InArgs,
 	const TSharedRef<STableViewBase>& OwnerTable,
@@ -101,6 +138,7 @@ void SInterchangePipelineItem::Construct(
 
 	STableRow<TSharedPtr<FInterchangePipelineItemType>>::Construct(
 		STableRow<TSharedPtr<FInterchangePipelineItemType>>::FArguments()
+		.Style(&FAppStyle::Get().GetWidgetStyle<FTableRowStyle>("SceneOutliner.TableViewRow"))
 		.Content()
 		[
 			SNew(SHorizontalBox)
@@ -297,6 +335,54 @@ void SInterchangePipelineConfigurationDialog::GatherConflictAndExtraInfo(TArray<
 	}
 }
 
+void SInterchangePipelineConfigurationDialog::OnExtendPipelineDetailsViewSettingsMenu(FMenuBuilder& MenuBuilder)
+{
+	FSlateIcon DummyIcon(NAME_None, NAME_None);
+
+	MenuBuilder.AddMenuEntry(
+		LOCTEXT("ShowEssentialsOnly_Label", "Show Only Essential Properties"),
+		LOCTEXT("ShowEssentialsOnly_Tooltip", "Display only the essential pipeline properties according to the pipeline settings"),
+		DummyIcon,
+		FUIAction(
+			FExecuteAction::CreateSP(this, &SInterchangePipelineConfigurationDialog::Execute_ToggleShowEssentials),
+			FCanExecuteAction(),
+			FIsActionChecked::CreateLambda([this]() {return bShowEssentials; })
+		),
+		FName("PipelineShowEssentials"),
+		EUserInterfaceActionType::Check
+	);
+
+	MenuBuilder.AddMenuEntry(
+		LOCTEXT("FilterOnContent_Label", "Filter Based on Import Contents"),
+		LOCTEXT("FilterOnContent_Tooltip", "Display properties based on the import source content"),
+		DummyIcon,
+		FUIAction(
+			FExecuteAction::CreateSP(this, &SInterchangePipelineConfigurationDialog::Execute_ToggleFilterOnContent),
+			FCanExecuteAction(),
+			FIsActionChecked::CreateLambda([this]() {return bFilterOptions; })
+		),
+		FName("PipelineShowEssentials"),
+		EUserInterfaceActionType::Check
+	);
+}
+
+void SInterchangePipelineConfigurationDialog::Execute_ToggleShowEssentials()
+{
+	bShowEssentials = !bShowEssentials;
+
+	//Refresh the pipeline
+	UpdateStack(CurrentStackName);
+	GConfig->SetBool(TEXT("InterchangeImportDialogOptions"), TEXT("ShowEssentials"), bShowEssentials, GEditorPerProjectIni);
+}
+
+void SInterchangePipelineConfigurationDialog::Execute_ToggleFilterOnContent()
+{
+	bFilterOptions = !bFilterOptions;
+	//Refresh the pipeline
+	UpdateStack(CurrentStackName);
+	GConfig->SetBool(TEXT("InterchangeImportDialogOptions"), TEXT("FilterOptions"), bFilterOptions, GEditorPerProjectIni);
+}
+
 TSharedRef<SBox> SInterchangePipelineConfigurationDialog::SpawnPipelineConfiguration()
 {
 	const ISlateStyle* InterchangeEditorPipelineStyle = FSlateStyleRegistry::FindSlateStyle("InterchangeEditorPipelineStyle");
@@ -363,26 +449,52 @@ TSharedRef<SBox> SInterchangePipelineConfigurationDialog::SpawnPipelineConfigura
 		.ClearSelectionOnClick(false)
 		.ToolTipText(PipelineListTooltip);
 
-	TSharedPtr<STextComboBox> TextComboBoxPtr;
+	TSharedPtr<SPipelinePresetComboBox> PresetComboBoxPtr;
 	//Only use a combo box if there is more then one stack
 	if (AvailableStacks.Num() > 0)
 	{
 		FText StackComboBoxTooltip = LOCTEXT("StackComboBoxTooltip", "Selected pipeline stack preset will be used for the current import. See the Interchange project settings to modify the pipeline stacks preset list.");
-		TextComboBoxPtr = SNew(STextComboBox)
+		PresetComboBoxPtr = SNew(SPipelinePresetComboBox)
 			.OptionsSource(&AvailableStacks)
+			.OnGenerateWidget(this, &SInterchangePipelineConfigurationDialog::MakePipelinePresetComboBoxEntryWidget)
 			.OnSelectionChanged(this, &SInterchangePipelineConfigurationDialog::OnStackSelectionChanged)
-			.ContentPadding(FMargin(0.0f, 2.0f))
-			.ToolTipText(StackComboBoxTooltip);
+			.ContentPadding(FMargin(2.0f, 4.0f))
+			.ToolTipText(StackComboBoxTooltip)
+			[
+				SNew(SHorizontalBox)
+				+SHorizontalBox::Slot()
+				.AutoWidth()
+				.VAlign(VAlign_Center)
+				.Padding(0.0, 0.0, 8.0, 0.0)
+				[
+					SNew(SBox)
+					.WidthOverride(14.0f)
+					.HeightOverride(14.0f)
+					[
+						SNew(SImage)
+						.Image(SInterchangePipelineStackPresetItem::GetItemIconBrush())
+					]
+				]
+				+ SHorizontalBox::Slot()
+				.FillWidth(1.0)
+				.VAlign(VAlign_Center)
+				.Padding(0.0, 0.0, 8.0, 0.0)
+				[
+					SNew(STextBlock)
+					.Text(this, &SInterchangePipelineConfigurationDialog::GetSelectedPipelineStackPresetLabelText)
+				]
+			];
+
 		if (SelectedStack.IsValid())
 		{
-			TextComboBoxPtr->SetSelectedItem(SelectedStack);
+			PresetComboBoxPtr->SetSelectedItem(SelectedStack);
 		}
 	}
 
-	FText CurrentStackText = LOCTEXT("CurrentStackText", "Stacks Preset");
+	FText CurrentStackText = LOCTEXT("CurrentStackText", "Pipeline Stack Preset");
 
 	TSharedPtr<SWidget> StackTextComboBox;
-	if (!TextComboBoxPtr.IsValid())
+	if (!PresetComboBoxPtr.IsValid())
 	{
 		CurrentStackText = LOCTEXT("CurrentStackTextNoComboBox", "There is no pipeline stack preset available");
 		StackTextComboBox = SNew(SBox)
@@ -409,9 +521,8 @@ TSharedRef<SBox> SInterchangePipelineConfigurationDialog::SpawnPipelineConfigura
 			]
 			+ SHorizontalBox::Slot()
 			.VAlign(VAlign_Center)
-			.AutoWidth()
 			[
-				TextComboBoxPtr.ToSharedRef()
+				PresetComboBoxPtr.ToSharedRef()
 			]
 		];
 	}
@@ -552,52 +663,6 @@ TSharedRef<SBox> SInterchangePipelineConfigurationDialog::SpawnPipelineConfigura
 				]
 				+SHorizontalBox::Slot()
 				.AutoWidth()
-				.Padding(4.0f, 0.0f)
-				[
-					SNew(SHorizontalBox)
-					+ SHorizontalBox::Slot()
-					.AutoWidth()
-					.VAlign(VAlign_Center)
-					.Padding(4.f, 0.f)
-					[
-						SNew(STextBlock)
-						.Text(LOCTEXT("SInterchangePipelineConfigurationDialog_ShowEssentialsOptions", "Essentials"))
-						.ToolTipText(EssentialPipelineTooltip)
-					]
-					+ SHorizontalBox::Slot()
-					.Padding(4.f, 0.f)
-					[
-						SNew(SCheckBox)
-						.IsChecked(this, &SInterchangePipelineConfigurationDialog::IsShowEssentialsEnabled)
-						.OnCheckStateChanged(this, &SInterchangePipelineConfigurationDialog::OnShowEssentialsChanged)
-						.ToolTipText(EssentialPipelineTooltip)
-					]
-				]
-				+SHorizontalBox::Slot()
-				.AutoWidth()
-				.Padding(4.0f, 0.0f)
-				[
-					SNew(SHorizontalBox)
-					+ SHorizontalBox::Slot()
-					.AutoWidth()
-					.VAlign(VAlign_Center)
-					.Padding(4.f, 0.f)
-					[
-						SNew(STextBlock)
-						.Text(LOCTEXT("SInterchangePipelineConfigurationDialog_FilterPipelineOptions", "Filter on Contents"))
-						.ToolTipText(FilterPipelineTooltip)
-					]
-					+ SHorizontalBox::Slot()
-					.Padding(4.f, 0.f)
-					[
-						SNew(SCheckBox)
-						.IsChecked(this, &SInterchangePipelineConfigurationDialog::IsFilteringOptions)
-						.OnCheckStateChanged(this, &SInterchangePipelineConfigurationDialog::OnFilterOptionsChanged)
-						.ToolTipText(FilterPipelineTooltip)
-					]
-				]
-				+SHorizontalBox::Slot()
-				.AutoWidth()
 				.VAlign(VAlign_Center)
 				.Padding(4.0f, 0.0f, 0.0f, 0.0f)
 				[
@@ -622,6 +687,29 @@ TSharedRef<SBox> SInterchangePipelineConfigurationDialog::SpawnPipelineConfigura
 				]
 			]
 			+ SVerticalBox::Slot()
+			.VAlign(VAlign_Center)
+			[
+				SNew(STextBlock)
+					.Visibility_Lambda([this]()
+						{
+							if (!CurrentSelectedStackPresetItem.IsValid())
+							{
+								return EVisibility::Collapsed;
+							}
+
+							return EVisibility::Visible;
+						})
+					.Text_Lambda([this]()
+						{
+							if (!CurrentSelectedStackPresetItem.IsValid())
+							{
+								return FText();
+							}
+
+							return FText::Format(LOCTEXT("PipelineListViewTitleFormat", "Pipelines in {0} pipeline stack:"), GetSelectedPipelineStackPresetLabelText());
+						})
+			]
+			+ SVerticalBox::Slot()
 			.Padding(0.0f, 8.0f)
 			.AutoHeight()
 			[
@@ -652,7 +740,15 @@ TSharedRef<SBox> SInterchangePipelineConfigurationDialog::SpawnPipelineConfigura
 	DetailsViewArgs.bShowHiddenPropertiesWhilePlayingOption = false;
 	DetailsViewArgs.NameAreaSettings = FDetailsViewArgs::HideNameArea;
 	DetailsViewArgs.ViewIdentifier = FName("InterchangeImportDialog");
-
+		
+	DetailsViewArgs.OptionsExtender = MakeShareable(new FExtender);
+	DetailsViewArgs.OptionsExtender->AddMenuExtension(
+		"DetailsViewShowOptions",
+		EExtensionHook::First,
+		nullptr, 
+		FMenuExtensionDelegate::CreateSP(this, &SInterchangePipelineConfigurationDialog::OnExtendPipelineDetailsViewSettingsMenu)
+	);
+	
 	PipelineConfigurationDetailsView = PropertyEditorModule.CreateDetailView(DetailsViewArgs);
 
 	// Reset to 'All' Sections on spawning the modal only. Subsequent changes to details view objects will use the persisted section selection.
@@ -802,6 +898,7 @@ void SInterchangePipelineConfigurationDialog::Construct(const FArguments& InArgs
 	ReimportObject = InArgs._ReimportObject;
 	SourceData = InArgs._SourceData;
 	Translator = InArgs._Translator;
+
 	if (Translator.IsValid())
 	{
 		TranslatorSettings = Translator->GetSettings();
@@ -862,6 +959,16 @@ void SInterchangePipelineConfigurationDialog::Construct(const FArguments& InArgs
 		GConfig->GetDouble(TEXT("InterchangeImportDialogOptions"), TEXT("SplitAdvancedRatio"), SplitAdvancedRatio, GEditorPerProjectIni);
 	}
 
+	if (InArgs._bOverrideDefaultShowEssentials.IsSet())
+	{
+		bShowEssentials = InArgs._bOverrideDefaultShowEssentials.GetValue();
+	}
+
+	if (InArgs._bOverrideDefaultFilterOnContent.IsSet())
+	{
+		bFilterOptions = InArgs._bOverrideDefaultFilterOnContent.GetValue();
+	}
+
 	//Make sure the windows is width enough to show all the ui part (cards and settings)
 	if (OwnerWindowPinned.IsValid())
 	{
@@ -915,7 +1022,7 @@ void SInterchangePipelineConfigurationDialog::Construct(const FArguments& InArgs
 				+ SHorizontalBox::Slot()
 				.AutoWidth()
 				.VAlign(VAlign_Center)
-				.Padding(8.0f, 0.0f, 8.0f, 0.0f)
+				.Padding(0.0f, 0.0f, 8.0f, 0.0f)
 				[
 					SNew(STextBlock)
 					.Text(LOCTEXT("SInterchangePipelineConfigurationDialog_SourceLabel", "Import Source"))
@@ -1451,12 +1558,29 @@ void SInterchangePipelineConfigurationDialog::UpdateStack(const FName& NewStackN
 	RefreshCardsViewList();
 }
 
+FText SInterchangePipelineConfigurationDialog::GetSelectedPipelineStackPresetLabelText() const
+{
+	if (CurrentSelectedStackPresetItem.IsValid())
+	{
+		return FText::FromString(*(CurrentSelectedStackPresetItem.Pin()));
+	}
+	return LOCTEXT("NoSelectedStackPreset", "No Selected Stack Preset");
+}
+
+TSharedRef<SWidget> SInterchangePipelineConfigurationDialog::MakePipelinePresetComboBoxEntryWidget(TSharedPtr<FString> InPresetItem)
+{
+	return SNew(SInterchangePipelineStackPresetItem)
+				.PresetItem(InPresetItem);
+}
+
 void SInterchangePipelineConfigurationDialog::OnStackSelectionChanged(TSharedPtr<FString> String, ESelectInfo::Type)
 {
 	if (!String.IsValid())
 	{
 		return;
 	}
+	
+	CurrentSelectedStackPresetItem = String.ToWeakPtr();
 
 	FName NewStackName = FName(*String.Get());
 	if (!UE::Private::ContainStack(PipelineStacks, NewStackName))
@@ -1513,6 +1637,12 @@ void SInterchangePipelineConfigurationDialog::UpdatePipelineSupportedAssetClasse
 			}
 		}
 	}
+
+	PipelineSupportedAssetData.Empty();
+	for (UClass* SupportedClass : PipelineSupportAssetClasses)
+	{
+		PipelineSupportedAssetData.Emplace(SupportedClass, TSet<UClass*>());
+	}
 }
 
 void SInterchangePipelineConfigurationDialog::UpdateEnableDataPerFactoryNodeClass()
@@ -1521,6 +1651,7 @@ void SInterchangePipelineConfigurationDialog::UpdateEnableDataPerFactoryNodeClas
 	{
 		return;
 	}
+
 	TArray<UClass*> ValidCardClasses;
 	PreviewNodeContainer->IterateNodesOfType<UInterchangeFactoryBaseNode>([this, &ValidCardClasses](const FString& NodeUid, UInterchangeFactoryBaseNode* FactoryNode)
 		{
@@ -1532,10 +1663,12 @@ void SInterchangePipelineConfigurationDialog::UpdateEnableDataPerFactoryNodeClas
 				&& NodeObjectClass)
 			{
 				bool bClassIsSupported = false;
+				UClass* AssetClass = nullptr;
 				for (UClass* SupportedClass : PipelineSupportAssetClasses)
 				{
 					if (NodeObjectClass->IsChildOf(SupportedClass))
 					{
+						AssetClass = SupportedClass;
 						bClassIsSupported = true;
 						break;
 					}
@@ -1543,6 +1676,8 @@ void SInterchangePipelineConfigurationDialog::UpdateEnableDataPerFactoryNodeClas
 				if (bClassIsSupported)
 				{
 					ValidCardClasses.Add(FactoryNode->GetClass());
+					PipelineSupportedAssetData[AssetClass].FactoryNodeClasses.Add(FactoryNode->GetClass());
+					
 					if (!EnableDataPerFactoryNodeClass.Contains(FactoryNode->GetClass()))
 					{
 						//Add a new entry that is enabled by default
@@ -1579,35 +1714,55 @@ void SInterchangePipelineConfigurationDialog::FillAssetCardsList()
 	{
 		UpdateEnableDataPerFactoryNodeClass();
 		AssetCards.Reset();
-		for (TPair<UClass*, FFactoryNodeEnabledData>& FactoryNodeClassAndEnableStatus : EnableDataPerFactoryNodeClass)
+		for (TPair<UClass*, FPerSupportedAssetClassData>& SupportedAssetDataPair : PipelineSupportedAssetData)
 		{
-			UClass* FactoryNodeClass = FactoryNodeClassAndEnableStatus.Key;
-			UClass* AssetClass = FactoryNodeClassAndEnableStatus.Value.ObjectClass;
-			if (AssetClass)
+			UClass* AssetClass = SupportedAssetDataPair.Key;
+			FPerSupportedAssetClassData AssetClassData = SupportedAssetDataPair.Value;
+
+			TSharedPtr<SInterchangeAssetCard> AssetCard;
+			if (AssetClassData.FactoryNodeClasses.IsEmpty())
 			{
-				TSharedPtr<SInterchangeAssetCard> AssetCard = SNew(SInterchangeAssetCard)
+				AssetCard = SNew(SInterchangeAssetCard)
 					.PreviewNodeContainer(PreviewNodeContainer)
 					.AssetClass(AssetClass)
-					.ShouldImportAssetType_Lambda([this, FactoryNodeClass]()
-						{
-							return EnableDataPerFactoryNodeClass.FindChecked(FactoryNodeClass).bEnable;
-						})
-					.OnImportAssetTypeChanged_Lambda([this, FactoryNodeClass](bool bNewEnabledValue)
-						{
-							EnableDataPerFactoryNodeClass.FindChecked(FactoryNodeClass).bEnable = bNewEnabledValue;
-						});
-
-				if (GInterchangeShowConflictWarningsOnCardsView
-					&& PipelineListViewItems.Num() > 0)
+					.CardDisabled(true)
+					.ShouldImportAssetType_Lambda([](){return false;})
+					.OnImportAssetTypeChanged_Lambda([](bool bNewEnabledValue){});
+			}
+			else
+			{
+				for (UClass* FactoryNodeClass : AssetClassData.FactoryNodeClasses)
 				{
-					for (TSharedPtr<FInterchangePipelineItemType>& PipelineListViewItem : PipelineListViewItems)
+					AssetClass = EnableDataPerFactoryNodeClass.FindChecked(FactoryNodeClass).ObjectClass;
+					AssetCard = SNew(SInterchangeAssetCard)
+						.PreviewNodeContainer(PreviewNodeContainer)
+						.AssetClass(AssetClass)
+						.CardDisabled(false)
+						.ShouldImportAssetType_Lambda([this, FactoryNodeClass]()
+							{
+								return EnableDataPerFactoryNodeClass.FindChecked(FactoryNodeClass).bEnable;
+							})
+						.OnImportAssetTypeChanged_Lambda([this, FactoryNodeClass](bool bNewEnabledValue)
+							{
+								EnableDataPerFactoryNodeClass.FindChecked(FactoryNodeClass).bEnable = bNewEnabledValue;
+							});
+
+					if (GInterchangeShowConflictWarningsOnCardsView
+						&& PipelineListViewItems.Num() > 0)
 					{
-						if (AssetCard->RefreshHasConflicts(PipelineListViewItem->ConflictInfos))
+						for (TSharedPtr<FInterchangePipelineItemType>& PipelineListViewItem : PipelineListViewItems)
 						{
-							break;
+							if (AssetCard->RefreshHasConflicts(PipelineListViewItem->ConflictInfos))
+							{
+								break;
+							}
 						}
 					}
 				}
+			}
+
+			if (AssetCard.IsValid())
+			{
 				AssetCards.Add(AssetCard);
 			}
 		}
@@ -1659,35 +1814,6 @@ UInterchangeCardsPipeline* SInterchangePipelineConfigurationDialog::GenerateTran
 	return InterchangeCardsPipeline;
 }
 
-void SInterchangePipelineConfigurationDialog::OnFilterOptionsChanged(ECheckBoxState CheckState)
-{
-	bool bNewCheckValue = CheckState == ECheckBoxState::Checked ? true : false;
-	if (bNewCheckValue == bFilterOptions)
-	{
-		//Check state did not change
-		return;
-	}
-	bFilterOptions = bNewCheckValue;
-	//Refresh the pipeline
-	UpdateStack(CurrentStackName);
-	GConfig->SetBool(TEXT("InterchangeImportDialogOptions"), TEXT("FilterOptions"), bFilterOptions, GEditorPerProjectIni);
-}
-
-void SInterchangePipelineConfigurationDialog::OnShowEssentialsChanged(ECheckBoxState CheckState)
-{
-	bool bNewCheckValue = CheckState == ECheckBoxState::Checked ? true : false;
-	if (bNewCheckValue == bShowEssentials)
-	{
-		//Check state did not change
-		return;
-	}
-	bShowEssentials = bNewCheckValue;
-	//Refresh the pipeline
-	UpdateStack(CurrentStackName);
-
-	GConfig->SetBool(TEXT("InterchangeImportDialogOptions"), TEXT("ShowEssentials"), bShowEssentials, GEditorPerProjectIni);
-}
-
 void SInterchangePipelineConfigurationDialog::UpdatePreviewContainer(bool bUpdateCards) const
 {
 	auto ClearObjectFlags = [](UObject* Obj)
diff --git a/Engine/Plugins/Interchange/Editor/Source/Pipelines/Private/SInterchangePipelineConfigurationDialog.h b/Engine/Plugins/Interchange/Editor/Source/Pipelines/Private/SInterchangePipelineConfigurationDialog.h
--- a/Engine/Plugins/Interchange/Editor/Source/Pipelines/Private/SInterchangePipelineConfigurationDialog.h
+++ b/Engine/Plugins/Interchange/Editor/Source/Pipelines/Private/SInterchangePipelineConfigurationDialog.h
@@ -7,6 +7,8 @@
 #include "InterchangePipelineConfigurationBase.h"
 #include "SInterchangeAssetCard.h"
 #include "Styling/SlateBrush.h"
+#include "Widgets/Input/SComboBox.h"
+#include "Widgets/Layout/SBox.h"
 #include "Widgets/Layout/SScrollBox.h"
 #include "Widgets/Views/STreeView.h"
 
@@ -16,11 +18,33 @@ struct FPropertyAndParent;
 struct FSlateBrush;
 class IDetailsView;
 class SCheckBox;
-class STextComboBox;
 class UInterchangeCardsPipeline;
 class UInterchangeTranslatorBase;
 
 
+class SInterchangePipelineConfigurationDialog;
+
+// Pipeline Stack Presets
+class SInterchangePipelineStackPresetItem : public SHorizontalBox
+{
+public:
+	SLATE_BEGIN_ARGS(SInterchangePipelineStackPresetItem)
+		: _PresetItem(nullptr)
+		{}
+		SLATE_ARGUMENT(TSharedPtr<FString>, PresetItem)		
+	SLATE_END_ARGS()
+
+	void Construct(const FArguments& InArgs);
+
+	static const FSlateBrush* GetItemIconBrush();
+
+private:
+	TSharedPtr<FString> PresetItem = nullptr;
+};
+
+typedef SComboBox< TSharedPtr<FString> > SPipelinePresetComboBox;
+
+// Pipeline List View 
 struct FInterchangePipelineItemType
 {
 public:
@@ -73,6 +97,8 @@ public:
 		SLATE_ARGUMENT(TWeakObjectPtr<UInterchangeBaseNodeContainer>, BaseNodeContainer)
 		SLATE_ARGUMENT(TWeakObjectPtr<UObject>, ReimportObject)
 		SLATE_ARGUMENT(TWeakObjectPtr<UInterchangeTranslatorBase>, Translator)
+		SLATE_ARGUMENT(TOptional<bool>, bOverrideDefaultShowEssentials)
+		SLATE_ARGUMENT(TOptional<bool>, bOverrideDefaultFilterOnContent)
 	SLATE_END_ARGS()
 
 public:
@@ -122,6 +148,12 @@ private:
 	FReply OnEditTranslatorSettings();
 
 	void GatherConflictAndExtraInfo(TArray<FInterchangeConflictInfo>& ConflictInfo, TMap<FString, FString>& ExtraInfo);
+
+	// Details View Menu Extension
+	void OnExtendPipelineDetailsViewSettingsMenu(FMenuBuilder& MenuBuilder);
+	void Execute_ToggleShowEssentials();
+	void Execute_ToggleFilterOnContent();
+
 private:
 	TWeakPtr< SWindow > OwnerWindow;
 	double OriginalMinWindowSize = 0.0;
@@ -137,6 +169,10 @@ private:
 
 	// The available stacks
 	TArray<TSharedPtr<FString>> AvailableStacks;
+
+	TSharedRef<SWidget> MakePipelinePresetComboBoxEntryWidget(TSharedPtr<FString> InElement);
+	FText GetSelectedPipelineStackPresetLabelText() const;
+
 	void OnStackSelectionChanged(TSharedPtr<FString> String, ESelectInfo::Type);
 	void UpdateStack(const FName& NewStackName);
 
@@ -159,6 +195,14 @@ private:
 		TObjectPtr<UClass> ObjectClass = nullptr;
 	};
 	TArray<TObjectPtr<UClass>> PipelineSupportAssetClasses;
+
+	struct FPerSupportedAssetClassData
+	{
+		TSet<UClass*> FactoryNodeClasses;
+	};
+
+	TMap<UClass*, FPerSupportedAssetClassData> PipelineSupportedAssetData;
+
 	TSharedPtr<SScrollBar> CardViewScrollbar;
 	TMap<UClass*, FFactoryNodeEnabledData> EnableDataPerFactoryNodeClass;
 	TSharedPtr<SInterchangeAssetCardList> CardViewList = nullptr;
@@ -174,19 +218,6 @@ private:
 	double SplitAdvancedRatio = 0.6;
 	TSharedPtr<SSplitter> CardsAndAdvancedSplitter = nullptr;
 
-	ECheckBoxState IsFilteringOptions() const
-	{
-		return bFilterOptions ? ECheckBoxState::Checked : ECheckBoxState::Unchecked;
-	}
-
-	ECheckBoxState IsShowEssentialsEnabled() const
-	{
-		return bShowEssentials ? ECheckBoxState::Checked : ECheckBoxState::Unchecked;
-	}
-
-	void OnFilterOptionsChanged(ECheckBoxState CheckState);
-	void OnShowEssentialsChanged(ECheckBoxState CheckState);
-
 	const FSlateBrush* GetImportButtonIcon() const;
 
 	void UpdatePreviewContainer(bool bUpdateCards) const;
@@ -212,6 +243,7 @@ private:
 	bool bShowCards = true;
 
 	FName CurrentStackName = NAME_None;
+	TWeakPtr<FString> CurrentSelectedStackPresetItem = nullptr;
 	TObjectPtr<UInterchangePipelineBase> CurrentSelectedPipeline = nullptr;
 	TWeakPtr<FInterchangePipelineItemType> CurrentSelectedPipelineItem = nullptr;
 };
diff --git a/Engine/Plugins/Interchange/Editor/Source/Pipelines/Public/InterchangePipelineConfigurationGeneric.h b/Engine/Plugins/Interchange/Editor/Source/Pipelines/Public/InterchangePipelineConfigurationGeneric.h
--- a/Engine/Plugins/Interchange/Editor/Source/Pipelines/Public/InterchangePipelineConfigurationGeneric.h
+++ b/Engine/Plugins/Interchange/Editor/Source/Pipelines/Public/InterchangePipelineConfigurationGeneric.h
@@ -11,6 +11,9 @@
 class UInterchangeBaseNodeContainer;
 class UInterchangeSourceData;
 
+class SWindow;
+class SInterchangePipelineConfigurationDialog;
+
 UCLASS(MinimalAPI, BlueprintType, editinlinenew)
 class UInterchangePipelineConfigurationGeneric : public UInterchangePipelineConfigurationBase
 {
@@ -19,32 +22,8 @@ class UInterchangePipelineConfigurationGeneric : public UInterchangePipelineConf
 public:
 
 protected:
+	UE_API virtual EInterchangePipelineConfigurationDialogResult ShowPipelineDialog_Internal(FPipelineConfigurationDialogParams& InParams) override;
 
-	UE_API virtual EInterchangePipelineConfigurationDialogResult ShowPipelineConfigurationDialog(TArray<FInterchangeStackInfo>& PipelineStacks
-		, TArray<UInterchangePipelineBase*>& OutPipelines
-		, TWeakObjectPtr<UInterchangeSourceData> SourceData
-		, TWeakObjectPtr <UInterchangeTranslatorBase> Translator
-		, TWeakObjectPtr<UInterchangeBaseNodeContainer> BaseNodeContainer) override;
-	UE_API virtual EInterchangePipelineConfigurationDialogResult ShowScenePipelineConfigurationDialog(TArray<FInterchangeStackInfo>& PipelineStacks
-		, TArray<UInterchangePipelineBase*>& OutPipelines
-		, TWeakObjectPtr<UInterchangeSourceData> SourceData
-		, TWeakObjectPtr <UInterchangeTranslatorBase> Translator
-		, TWeakObjectPtr<UInterchangeBaseNodeContainer> BaseNodeContainer) override;
-	UE_API virtual EInterchangePipelineConfigurationDialogResult ShowReimportPipelineConfigurationDialog(TArray<FInterchangeStackInfo>& PipelineStacks
-		, TArray<UInterchangePipelineBase*>& OutPipelines
-		, TWeakObjectPtr<UInterchangeSourceData> SourceData
-		, TWeakObjectPtr <UInterchangeTranslatorBase> Translator
-		, TWeakObjectPtr<UInterchangeBaseNodeContainer> BaseNodeContainer
-		, TWeakObjectPtr <UObject> ReimportAsset
-		, bool bSceneImport) override;
-	UE_API virtual EInterchangePipelineConfigurationDialogResult ShowTestPlanConfigurationDialog(TArray<FInterchangeStackInfo>& PipelineStacks
-		, TArray<UInterchangePipelineBase*>& OutPipelines
-		, TWeakObjectPtr<UInterchangeSourceData> SourceData
-		, TWeakObjectPtr <UInterchangeTranslatorBase> Translator
-		, TWeakObjectPtr<UInterchangeBaseNodeContainer> BaseNodeContainer
-		, TWeakObjectPtr <UObject> ReimportAsset
-		, bool bSceneImport
-		, bool bReimport) override;
 };
 
 #undef UE_API
diff --git a/Engine/Plugins/Interchange/Runtime/Source/Pipelines/Private/InterchangeGenericMeshPipeline.cpp b/Engine/Plugins/Interchange/Runtime/Source/Pipelines/Private/InterchangeGenericMeshPipeline.cpp
--- a/Engine/Plugins/Interchange/Runtime/Source/Pipelines/Private/InterchangeGenericMeshPipeline.cpp
+++ b/Engine/Plugins/Interchange/Runtime/Source/Pipelines/Private/InterchangeGenericMeshPipeline.cpp
@@ -314,6 +314,11 @@ void UInterchangeGenericMeshPipeline::GetSupportAssetClasses(TArray<UClass*>& Pi
 	{
 		PipelineSupportAssetClasses.Add(UPhysicsAsset::StaticClass());
 	}
+
+	if (CommonSkeletalMeshesAndAnimationsProperties.IsValid() && CommonSkeletalMeshesAndAnimationsProperties->Skeleton == nullptr)
+	{
+		PipelineSupportAssetClasses.Add(USkeleton::StaticClass());
+	}
 }
 
 #endif
diff --git a/Engine/Source/Runtime/Interchange/Engine/Private/InterchangeSceneImportAsset.cpp b/Engine/Source/Runtime/Interchange/Engine/Private/InterchangeSceneImportAsset.cpp
--- a/Engine/Source/Runtime/Interchange/Engine/Private/InterchangeSceneImportAsset.cpp
+++ b/Engine/Source/Runtime/Interchange/Engine/Private/InterchangeSceneImportAsset.cpp
@@ -327,3 +327,11 @@ bool UInterchangeSceneImportAsset::ContainsObject(const FSoftObjectPath& SoftObj
 	return false;
 #endif
 }
+
+
+#if WITH_EDITOR
+const TMap<FSoftObjectPath, FString> UInterchangeSceneImportAsset::GetCachedSceneObjects() const
+{
+	return SceneObjects;
+}
+#endif 
\ No newline at end of file
diff --git a/Engine/Source/Runtime/Interchange/Engine/Public/InterchangePipelineConfigurationBase.h b/Engine/Source/Runtime/Interchange/Engine/Public/InterchangePipelineConfigurationBase.h
--- a/Engine/Source/Runtime/Interchange/Engine/Public/InterchangePipelineConfigurationBase.h
+++ b/Engine/Source/Runtime/Interchange/Engine/Public/InterchangePipelineConfigurationBase.h
@@ -8,6 +8,7 @@
 #include "InterchangeSourceData.h"
 #include "InterchangeTranslatorBase.h"
 #include "Nodes/InterchangeBaseNodeContainer.h"
+#include "Misc/Optional.h"
 
 #include "InterchangePipelineConfigurationBase.generated.h"
 
@@ -57,7 +58,18 @@ public:
 		, UInterchangeBaseNodeContainer* BaseNodeContainer)
 	{
 		//By default we call the virtual import pipeline execution
-		return ShowPipelineConfigurationDialog(PipelineStacks, OutPipelines, SourceData, Translator, BaseNodeContainer);
+		FPipelineConfigurationDialogParams DialogParams{
+			.PipelineStacks = PipelineStacks,
+			.OutPipelines = OutPipelines,
+			.SourceData = SourceData,
+			.Translator = Translator,
+			.BaseNodeContainer = BaseNodeContainer,
+			.ReimportAsset = nullptr,
+			.bReimport = false,
+			.bSceneImport = false,
+			.bInvokedThroughTestPlan = false,
+		};
+		return ShowPipelineDialog_Internal(DialogParams);
 	}
 
 	/**
@@ -78,7 +90,19 @@ public:
 		, UInterchangeBaseNodeContainer* BaseNodeContainer)
 	{
 		//By default we call the virtual import pipeline execution
-		return ShowScenePipelineConfigurationDialog(PipelineStacks, OutPipelines, SourceData, Translator, BaseNodeContainer);
+		FPipelineConfigurationDialogParams DialogParams{
+			.PipelineStacks = PipelineStacks,
+			.OutPipelines = OutPipelines,
+			.SourceData = SourceData,
+			.Translator = Translator,
+			.BaseNodeContainer = BaseNodeContainer,
+			.ReimportAsset = nullptr,
+			.bReimport = false,
+			.bSceneImport = true,
+			.bInvokedThroughTestPlan = false,
+		};
+
+		return ShowPipelineDialog_Internal(DialogParams);
 	}
 
 	/**
@@ -103,7 +127,19 @@ public:
 		, bool bSceneImport = false)
 	{
 		//By default we call the virtual import pipeline execution
-		return ShowReimportPipelineConfigurationDialog(PipelineStacks, OutPipelines, SourceData, Translator, BaseNodeContainer, ReimportAsset, bSceneImport);
+		FPipelineConfigurationDialogParams DialogParams{
+			.PipelineStacks = PipelineStacks,
+			.OutPipelines = OutPipelines,
+			.SourceData = SourceData,
+			.Translator = Translator,
+			.BaseNodeContainer = BaseNodeContainer,
+			.ReimportAsset = ReimportAsset,
+			.bReimport = true,
+			.bSceneImport = bSceneImport,
+			.bInvokedThroughTestPlan = false,
+		};
+
+		return ShowPipelineDialog_Internal(DialogParams);
 	}
 
 	/**
@@ -130,63 +166,38 @@ public:
 		, bool bReimport = false)
 	{
 		//By default we call the virtual import pipeline execution
-		return ShowTestPlanConfigurationDialog(PipelineStacks, OutPipelines, SourceData, Translator, BaseNodeContainer, ReimportAsset, bSceneImport, bReimport);
+		FPipelineConfigurationDialogParams DialogParams{
+			.PipelineStacks = PipelineStacks,
+			.OutPipelines = OutPipelines,
+			.SourceData = SourceData,
+			.Translator = Translator,
+			.BaseNodeContainer = BaseNodeContainer,
+			.ReimportAsset = ReimportAsset,
+			.bReimport = bReimport,
+			.bSceneImport = bSceneImport,
+			.bInvokedThroughTestPlan = true,
+		};
+
+		return ShowPipelineDialog_Internal(DialogParams);
 	}
 
 protected:
-
-	/**
-	 * This function shows a dialog used to configure pipeline stacks and returns a stack name that tells the caller the user's choice.
-	 */
-	virtual EInterchangePipelineConfigurationDialogResult ShowPipelineConfigurationDialog(TArray<FInterchangeStackInfo>& PipelineStacks
-		, TArray<UInterchangePipelineBase*>& OutPipelines
-		, TWeakObjectPtr<UInterchangeSourceData> SourceData
-		, TWeakObjectPtr <UInterchangeTranslatorBase> Translator
-		, TWeakObjectPtr <UInterchangeBaseNodeContainer> BaseNodeContainer)
-	{ 
-		//Not implemented
-		return EInterchangePipelineConfigurationDialogResult::Cancel;
-	}
-
-	/**
-	 * This function shows a dialog used to configure pipeline stacks and returns a stack name that tells the caller the user's choice.
-	 */
-	virtual EInterchangePipelineConfigurationDialogResult ShowScenePipelineConfigurationDialog(TArray<FInterchangeStackInfo>& PipelineStacks
-		, TArray<UInterchangePipelineBase*>& OutPipelines
-		, TWeakObjectPtr<UInterchangeSourceData> SourceData
-		, TWeakObjectPtr <UInterchangeTranslatorBase> Translator
-		, TWeakObjectPtr <UInterchangeBaseNodeContainer> BaseNodeContainer)
-	{ 
-		//Not implemented
-		return EInterchangePipelineConfigurationDialogResult::Cancel;
-	}
-
-	/**
-	 * This function shows a dialog used to configure pipeline stacks and returns a stack name that tells the caller the user's choice.
-	 */
-	virtual EInterchangePipelineConfigurationDialogResult ShowReimportPipelineConfigurationDialog(TArray<FInterchangeStackInfo>& PipelineStacks
-		, TArray<UInterchangePipelineBase*>& OutPipelines
-		, TWeakObjectPtr<UInterchangeSourceData> SourceData
-		, TWeakObjectPtr <UInterchangeTranslatorBase> Translator
-		, TWeakObjectPtr <UInterchangeBaseNodeContainer> BaseNodeContainer
-		, TWeakObjectPtr <UObject> ReimportAsset
-		, bool bSceneImport)
+	struct FPipelineConfigurationDialogParams
 	{
-		//Not implemented
-		return EInterchangePipelineConfigurationDialogResult::Cancel;
-	}
-
-	/**
-	 * This function shows a dialog used to configure pipeline stacks in interchange test plan asset.
-	 */
-	virtual EInterchangePipelineConfigurationDialogResult ShowTestPlanConfigurationDialog(TArray<FInterchangeStackInfo>& PipelineStacks
-		, TArray<UInterchangePipelineBase*>& OutPipelines
-		, TWeakObjectPtr<UInterchangeSourceData> SourceData
-		, TWeakObjectPtr <UInterchangeTranslatorBase> Translator
-		, TWeakObjectPtr <UInterchangeBaseNodeContainer> BaseNodeContainer
-		, TWeakObjectPtr <UObject> ReimportAsset
-		, bool bSceneImport
-		, bool bReimport)
+		TArray<FInterchangeStackInfo>& PipelineStacks;
+		TArray<UInterchangePipelineBase*>& OutPipelines;
+		TWeakObjectPtr<UInterchangeSourceData> SourceData;
+		TWeakObjectPtr <UInterchangeTranslatorBase> Translator;
+		TWeakObjectPtr <UInterchangeBaseNodeContainer> BaseNodeContainer;
+		TWeakObjectPtr <UObject> ReimportAsset;
+		bool bReimport;
+		bool bSceneImport;
+		bool bInvokedThroughTestPlan;
+		TOptional<bool> bOverrideDefaultShowEssentials;
+		TOptional<bool> bOverrideDefaultFilterOnContent;
+	};
+
+	virtual EInterchangePipelineConfigurationDialogResult ShowPipelineDialog_Internal(FPipelineConfigurationDialogParams& InParams)
 	{
 		//Not implemented
 		return EInterchangePipelineConfigurationDialogResult::Cancel;
