[Interchange] Temp fix for Interchange Logging
- Enabled logs based on the Automation Testing condition in the Interchange Manager
- Reenabled the Test Plan that expects errors. Succeeded when tested in the Automate Test Frontend.

#rnx

diff --git a/Engine/Plugins/Interchange/Editor/Source/InterchangeEditor/Private/InterchangeEditorModule.cpp b/Engine/Plugins/Interchange/Editor/Source/InterchangeEditor/Private/InterchangeEditorModule.cpp
--- a/Engine/Plugins/Interchange/Editor/Source/InterchangeEditor/Private/InterchangeEditorModule.cpp
+++ b/Engine/Plugins/Interchange/Editor/Source/InterchangeEditor/Private/InterchangeEditorModule.cpp
@@ -68,7 +68,7 @@ namespace UE::Interchange::InterchangeEditorModule
 			}
 
 			// UInterchangeManager::ReleaseAsyncHelper has output logs turned on. To avoid duplication, turning it off here.
-			constexpr bool bMirrorToOutputLog = false;
+			const bool bMirrorToOutputLog = GIsAutomationTesting;
 			LogListing->AddMessages(TokenizedMessages, bMirrorToOutputLog);
 			LogListing->NotifyIfAnyMessages(NSLOCTEXT("Interchange", "LogAndNotify", "There were issues with the import."), EMessageSeverity::Info);
 		}
diff --git a/Engine/Source/Runtime/Interchange/Engine/Private/InterchangeManager.cpp b/Engine/Source/Runtime/Interchange/Engine/Private/InterchangeManager.cpp
--- a/Engine/Source/Runtime/Interchange/Engine/Private/InterchangeManager.cpp
+++ b/Engine/Source/Runtime/Interchange/Engine/Private/InterchangeManager.cpp
@@ -2741,8 +2741,10 @@ void UInterchangeManager::ReleaseAsyncHelper(TWeakPtr<UE::Interchange::FImportAs
 	using namespace UE::Interchange;
 
 	check(AsyncHelper.IsValid());
-	
-	constexpr bool bLogWarningsAndErrors = true;
+		
+	// TODO: @vjavdekar : Find a way to run automated tests so that having this enabled
+	// doesn't interrupt them and return early causing some to fail that expect the errors. 
+	const bool bLogWarningsAndErrors = !GIsAutomationTesting;
 	
 	double ImportTimeDelta;
 	bool bSucceeded = false;
