Tentatively fixed  crash reported by users when closing import dialog
- Solution: Replaced TObjectPtr on preview node container with a TWeakObjectPtr

#rnx

diff --git a/Engine/Plugins/Interchange/Editor/Source/Pipelines/Private/SInterchangePipelineConfigurationDialog.cpp b/Engine/Plugins/Interchange/Editor/Source/Pipelines/Private/SInterchangePipelineConfigurationDialog.cpp
--- a/Engine/Plugins/Interchange/Editor/Source/Pipelines/Private/SInterchangePipelineConfigurationDialog.cpp
+++ b/Engine/Plugins/Interchange/Editor/Source/Pipelines/Private/SInterchangePipelineConfigurationDialog.cpp
@@ -224,11 +224,11 @@ SInterchangePipelineConfigurationDialog::~SInterchangePipelineConfigurationDialo
 		OwnerWindowPinned->GetOnWindowClosedEvent().RemoveAll(this);
 	}
 
-	if (PreviewNodeContainer)
+	if (PreviewNodeContainer.IsValid())
 	{
 		PreviewNodeContainer->ClearFlags(RF_Standalone | RF_Public);
 		PreviewNodeContainer->ClearInternalFlags(EInternalObjectFlags::Async);
-		PreviewNodeContainer = nullptr;
+		PreviewNodeContainer.Reset();
 	}
 }
 
@@ -1647,7 +1647,7 @@ void SInterchangePipelineConfigurationDialog::UpdatePipelineSupportedAssetClasse
 
 void SInterchangePipelineConfigurationDialog::UpdateEnableDataPerFactoryNodeClass()
 {
-	if (!ensure(PreviewNodeContainer))
+	if (!ensure(PreviewNodeContainer.IsValid()))
 	{
 		return;
 	}
@@ -1710,7 +1710,7 @@ void SInterchangePipelineConfigurationDialog::FillAssetCardsList()
 	//Update the preview container
 	constexpr bool bUpdateCardsTrue = true;
 	UpdatePreviewContainer(bUpdateCardsTrue);
-	if (ensure(PreviewNodeContainer))
+	if (ensure(PreviewNodeContainer.IsValid()))
 	{
 		UpdateEnableDataPerFactoryNodeClass();
 		AssetCards.Reset();
@@ -1723,7 +1723,7 @@ void SInterchangePipelineConfigurationDialog::FillAssetCardsList()
 			if (AssetClassData.FactoryNodeClasses.IsEmpty())
 			{
 				AssetCard = SNew(SInterchangeAssetCard)
-					.PreviewNodeContainer(PreviewNodeContainer)
+					.PreviewNodeContainer(PreviewNodeContainer.Get())
 					.AssetClass(AssetClass)
 					.CardDisabled(true)
 					.ShouldImportAssetType_Lambda([](){return false;})
@@ -1735,7 +1735,7 @@ void SInterchangePipelineConfigurationDialog::FillAssetCardsList()
 				{
 					AssetClass = EnableDataPerFactoryNodeClass.FindChecked(FactoryNodeClass).ObjectClass;
 					AssetCard = SNew(SInterchangeAssetCard)
-						.PreviewNodeContainer(PreviewNodeContainer)
+						.PreviewNodeContainer(PreviewNodeContainer.Get())
 						.AssetClass(AssetClass)
 						.CardDisabled(false)
 						.ShouldImportAssetType_Lambda([this, FactoryNodeClass]()
@@ -1787,7 +1787,7 @@ void SInterchangePipelineConfigurationDialog::RefreshCardsViewList()
 	FillAssetCardsList();
 	if (CardViewList.IsValid())
 	{
-		CardViewList->RefreshList(PreviewNodeContainer);
+		CardViewList->RefreshList(PreviewNodeContainer.Get());
 	}
 }
 
@@ -1821,10 +1821,10 @@ void SInterchangePipelineConfigurationDialog::UpdatePreviewContainer(bool bUpdat
 			Obj->ClearFlags(RF_Standalone | RF_Public);
 			Obj->ClearInternalFlags(EInternalObjectFlags::Async);
 		};
-	if (PreviewNodeContainer)
+	if (PreviewNodeContainer.IsValid())
 	{
-		ClearObjectFlags(PreviewNodeContainer);
-		PreviewNodeContainer = nullptr;
+		ClearObjectFlags(PreviewNodeContainer.Get());
+		PreviewNodeContainer.Reset();
 	}
 	PreviewNodeContainer = DuplicateObject<UInterchangeBaseNodeContainer>(BaseNodeContainer.Get(), GetTransientPackage());
 	PreviewNodeContainer->SetChildrenCache(BaseNodeContainer->GetChildrenCache());
@@ -1843,7 +1843,7 @@ void SInterchangePipelineConfigurationDialog::UpdatePreviewContainer(bool bUpdat
 		{
 			DuplicatedPipeline->TransferAdjustSettings(PipelineItem->Pipeline);
 			DuplicatedPipeline->SetResultsContainer(Results);
-			DuplicatedPipeline->ScriptedExecutePipeline(PreviewNodeContainer, SourceDatas, FString());
+			DuplicatedPipeline->ScriptedExecutePipeline(PreviewNodeContainer.Get(), SourceDatas, FString());
 			ClearObjectFlags(DuplicatedPipeline);
 		}
 	}
@@ -1853,7 +1853,7 @@ void SInterchangePipelineConfigurationDialog::UpdatePreviewContainer(bool bUpdat
 		//If we do not update cards execute the cards pipeline since its a final preview
 		if (UInterchangeCardsPipeline* InterchangeCardsPipeline = GenerateTransientCardsPipeline())
 		{
-			InterchangeCardsPipeline->ScriptedExecutePipeline(PreviewNodeContainer, SourceDatas, FString());
+			InterchangeCardsPipeline->ScriptedExecutePipeline(PreviewNodeContainer.Get(), SourceDatas, FString());
 		}
 	}
 
@@ -1885,7 +1885,7 @@ FReply SInterchangePipelineConfigurationDialog::OnPreviewImport() const
 	constexpr bool bUpdateCardsFalse = false;
 	UpdatePreviewContainer(bUpdateCardsFalse);
 
-	if (!ensure(PreviewNodeContainer))
+	if (!ensure(PreviewNodeContainer.IsValid()))
 	{
 		return FReply::Handled();
 	}
@@ -1899,7 +1899,7 @@ FReply SInterchangePipelineConfigurationDialog::OnPreviewImport() const
 	Window->SetContent
 	(
 		SAssignNew(InterchangeGraphInspectorWindow, SInterchangeGraphInspectorWindow)
-		.InterchangeBaseNodeContainer(PreviewNodeContainer)
+		.InterchangeBaseNodeContainer(PreviewNodeContainer.Get())
 		.bPreview(true)
 		.OwnerWindow(Window)
 	);
diff --git a/Engine/Plugins/Interchange/Editor/Source/Pipelines/Private/SInterchangePipelineConfigurationDialog.h b/Engine/Plugins/Interchange/Editor/Source/Pipelines/Private/SInterchangePipelineConfigurationDialog.h
--- a/Engine/Plugins/Interchange/Editor/Source/Pipelines/Private/SInterchangePipelineConfigurationDialog.h
+++ b/Engine/Plugins/Interchange/Editor/Source/Pipelines/Private/SInterchangePipelineConfigurationDialog.h
@@ -160,7 +160,7 @@ private:
 	double DeltaClientWindowSize = 0.0;
 	TWeakObjectPtr<UInterchangeSourceData> SourceData;
 	TWeakObjectPtr<UInterchangeBaseNodeContainer> BaseNodeContainer;
-	mutable TObjectPtr<UInterchangeBaseNodeContainer> PreviewNodeContainer = nullptr;
+	mutable TWeakObjectPtr<UInterchangeBaseNodeContainer> PreviewNodeContainer;
 	TWeakObjectPtr<UObject> ReimportObject;
 	TWeakObjectPtr<UInterchangeTranslatorBase> Translator;
 	TObjectPtr<UInterchangeTranslatorSettings> TranslatorSettings = nullptr;
