[BUGFIX][Interchange] FBX Python Level Import Test Failing
- Changed the UWorld Type created for Level Instances to EWorldType::Inactive
- Added a guard on GIsUnattendedScript during automation testing to ensure that the package saving works.

#rnx

diff --git a/Engine/Plugins/Interchange/Editor/Source/Utilities/Private/InterchangeEditorUtilities.cpp b/Engine/Plugins/Interchange/Editor/Source/Utilities/Private/InterchangeEditorUtilities.cpp
--- a/Engine/Plugins/Interchange/Editor/Source/Utilities/Private/InterchangeEditorUtilities.cpp
+++ b/Engine/Plugins/Interchange/Editor/Source/Utilities/Private/InterchangeEditorUtilities.cpp
@@ -17,7 +17,16 @@ bool UInterchangeEditorUtilities::SaveAsset(UObject* Asset) const
 {
 	TArray<UPackage*> PackagesToSave;
 	PackagesToSave.Add(Asset->GetPackage());
-	FEditorFileUtils::EPromptReturnCode ReturnCode = FEditorFileUtils::PromptForCheckoutAndSave(PackagesToSave, false /*bCheckDirty*/, false /*bPromptToSave*/);
+	FEditorFileUtils::EPromptReturnCode ReturnCode;
+	if (GIsAutomationTesting)
+	{
+		TGuardValue<bool> RunningUnattendedScriptGuard(GIsRunningUnattendedScript, true);
+		ReturnCode = FEditorFileUtils::PromptForCheckoutAndSave(PackagesToSave, false /*bCheckDirty*/, false /*bPromptToSave*/);
+	}
+	else
+	{
+		ReturnCode = FEditorFileUtils::PromptForCheckoutAndSave(PackagesToSave, false /*bCheckDirty*/, false /*bPromptToSave*/);
+	}
 	return (ReturnCode == FEditorFileUtils::PR_Success);
 }
 
diff --git a/Engine/Plugins/Interchange/Runtime/Source/Import/Private/Scene/InterchangeLevelFactory.cpp b/Engine/Plugins/Interchange/Runtime/Source/Import/Private/Scene/InterchangeLevelFactory.cpp
--- a/Engine/Plugins/Interchange/Runtime/Source/Import/Private/Scene/InterchangeLevelFactory.cpp
+++ b/Engine/Plugins/Interchange/Runtime/Source/Import/Private/Scene/InterchangeLevelFactory.cpp
@@ -87,7 +87,7 @@ namespace UE::Interchange::Private::InterchangeLevelFactory
 				.CreateWorldPartition(bCreateWorldPartition)
 				.EnableWorldPartitionStreaming(bEnableWorldPartitionStreaming);
 
-			UWorld* NewWorld = UWorld::CreateWorld(EWorldType::Editor, bInformEngineOfWorld, *Arguments.AssetName, Cast<UPackage>(Arguments.Parent), bAddToRoot, ERHIFeatureLevel::Num, &InitValues);
+			UWorld* NewWorld = UWorld::CreateWorld(EWorldType::Inactive, bInformEngineOfWorld, *Arguments.AssetName, Cast<UPackage>(Arguments.Parent), bAddToRoot, ERHIFeatureLevel::Num, &InitValues);
 			NewWorld->SetFlags(RF_Public | RF_Standalone);
 			TargetAsset = NewWorld;
 		}
