[Interchange UI] - Add separator section headings in the Import Dialog details view panel settings drop down

#rnx

diff --git a/Engine/Plugins/Interchange/Editor/Source/Pipelines/Private/SInterchangePipelineConfigurationDialog.cpp b/Engine/Plugins/Interchange/Editor/Source/Pipelines/Private/SInterchangePipelineConfigurationDialog.cpp
--- a/Engine/Plugins/Interchange/Editor/Source/Pipelines/Private/SInterchangePipelineConfigurationDialog.cpp
+++ b/Engine/Plugins/Interchange/Editor/Source/Pipelines/Private/SInterchangePipelineConfigurationDialog.cpp
@@ -339,6 +339,7 @@ void SInterchangePipelineConfigurationDialog::OnExtendPipelineDetailsViewSetting
 {
 	FSlateIcon DummyIcon(NAME_None, NAME_None);
 
+	MenuBuilder.BeginSection("Interchange", LOCTEXT("InterchangeMenuExtenderHeading", "Interchange"));
 	MenuBuilder.AddMenuEntry(
 		LOCTEXT("ShowEssentialsOnly_Label", "Show Only Essential Properties"),
 		LOCTEXT("ShowEssentialsOnly_Tooltip", "Display only the essential pipeline properties according to the pipeline settings"),
@@ -364,6 +365,7 @@ void SInterchangePipelineConfigurationDialog::OnExtendPipelineDetailsViewSetting
 		FName("PipelineShowEssentials"),
 		EUserInterfaceActionType::Check
 	);
+	MenuBuilder.EndSection();
 }
 
 void SInterchangePipelineConfigurationDialog::Execute_ToggleShowEssentials()
@@ -740,15 +742,16 @@ TSharedRef<SBox> SInterchangePipelineConfigurationDialog::SpawnPipelineConfigura
 	DetailsViewArgs.bShowHiddenPropertiesWhilePlayingOption = false;
 	DetailsViewArgs.NameAreaSettings = FDetailsViewArgs::HideNameArea;
 	DetailsViewArgs.ViewIdentifier = FName("InterchangeImportDialog");
-		
+	
 	DetailsViewArgs.OptionsExtender = MakeShareable(new FExtender);
 	DetailsViewArgs.OptionsExtender->AddMenuExtension(
 		"DetailsViewShowOptions",
-		EExtensionHook::First,
+		EExtensionHook::Before,
 		nullptr, 
 		FMenuExtensionDelegate::CreateSP(this, &SInterchangePipelineConfigurationDialog::OnExtendPipelineDetailsViewSettingsMenu)
 	);
 	
+	
 	PipelineConfigurationDetailsView = PropertyEditorModule.CreateDetailView(DetailsViewArgs);
 
 	// Reset to 'All' Sections on spawning the modal only. Subsequent changes to details view objects will use the persisted section selection.
diff --git a/Engine/Source/Editor/PropertyEditor/Private/SDetailsView.cpp b/Engine/Source/Editor/PropertyEditor/Private/SDetailsView.cpp
--- a/Engine/Source/Editor/PropertyEditor/Private/SDetailsView.cpp
+++ b/Engine/Source/Editor/PropertyEditor/Private/SDetailsView.cpp
@@ -115,8 +115,8 @@ void SDetailsView::Construct(const FArguments& InArgs, const FDetailsViewArgs& I
 	ExternalScrollbar->SetVisibility( TAttribute<EVisibility>( this, &SDetailsView::GetScrollBarVisibility ) );
 
 	FMenuBuilder DetailViewOptions(true, nullptr, InDetailsViewArgs.OptionsExtender);
-
-	DetailViewOptions.BeginSection("DetailsViewShowOptions");
+	const bool bDisplayShowSectionHeadingText = InDetailsViewArgs.OptionsExtender.IsValid() && (InDetailsViewArgs.OptionsExtender->NumExtensions() > 0);
+	DetailViewOptions.BeginSection("DetailsViewShowOptions", (bDisplayShowSectionHeadingText ? LOCTEXT("ShowSectionHeading", "Show") : FText::GetEmpty()));
 
 	if (DetailsViewArgs.bShowModifiedPropertiesOption)
 	{
