Interchange: Fix small logic bug on the inspector pipeline UI preventing the "All" checkbox from affecting top-level nodes

diff --git a/Engine/Plugins/Interchange/Editor/Source/Pipelines/Private/SInterchangeGraphInspectorWindow.cpp b/Engine/Plugins/Interchange/Editor/Source/Pipelines/Private/SInterchangeGraphInspectorWindow.cpp
--- a/Engine/Plugins/Interchange/Editor/Source/Pipelines/Private/SInterchangeGraphInspectorWindow.cpp
+++ b/Engine/Plugins/Interchange/Editor/Source/Pipelines/Private/SInterchangeGraphInspectorWindow.cpp
@@ -256,13 +256,19 @@ void SInterchangeGraphInspectorTreeView::OnGetChildrenGraphInspectorTreeView(UIn
 
 void SInterchangeGraphInspectorTreeView::RecursiveSetImport(UInterchangeBaseNode* Node, bool bState)
 {
+	if (!Node)
+	{
+		return;
+	}
+
+	Node->SetEnabled(bState);
+
 	TArray<FString> Childrens = InterchangeBaseNodeContainer->GetNodeChildrenUids(Node->GetUniqueID());
 	for (const FString& ChildID : Childrens)
 	{
 		UInterchangeBaseNode* ChildNode = const_cast<UInterchangeBaseNode*>(InterchangeBaseNodeContainer->GetNode(ChildID));
 		if (!ChildNode)
 			continue;
-		ChildNode->SetEnabled(bState);
 		RecursiveSetImport(ChildNode, bState);
 	}
 }
