Interchange: Sort nodes by type and name in the inspector and preview panel

diff --git a/Engine/Plugins/Interchange/Editor/Source/Pipelines/InterchangeEditorPipelines.Build.cs b/Engine/Plugins/Interchange/Editor/Source/Pipelines/InterchangeEditorPipelines.Build.cs
--- a/Engine/Plugins/Interchange/Editor/Source/Pipelines/InterchangeEditorPipelines.Build.cs
+++ b/Engine/Plugins/Interchange/Editor/Source/Pipelines/InterchangeEditorPipelines.Build.cs
@@ -23,6 +23,7 @@ namespace UnrealBuildTool.Rules
 					"AssetTools",
 					"InputCore",
 					"InterchangePipelines",
+					"InterchangeNodes",
 					"MainFrame",
 					"Projects",
 					"PropertyEditor",
diff --git a/Engine/Plugins/Interchange/Editor/Source/Pipelines/Private/SInterchangeGraphInspectorWindow.cpp b/Engine/Plugins/Interchange/Editor/Source/Pipelines/Private/SInterchangeGraphInspectorWindow.cpp
--- a/Engine/Plugins/Interchange/Editor/Source/Pipelines/Private/SInterchangeGraphInspectorWindow.cpp
+++ b/Engine/Plugins/Interchange/Editor/Source/Pipelines/Private/SInterchangeGraphInspectorWindow.cpp
@@ -1,15 +1,18 @@
 // Copyright Epic Games, Inc. All Rights Reserved.
+
 #include "SInterchangeGraphInspectorWindow.h"
 
+#include "InterchangeSceneNode.h"
+#include "Nodes/InterchangeBaseNodeContainer.h"
+#include "Nodes/InterchangeFactoryBaseNode.h"
+
 #include "DetailsViewArgs.h"
-#include "Framework/MultiBox/MultiBoxBuilder.h"
 #include "Framework/Commands/UICommandList.h"
-#include "GameFramework/Actor.h"
+#include "Framework/MultiBox/MultiBoxBuilder.h"
 #include "Framework/Views/TableViewMetadata.h"
+#include "GameFramework/Actor.h"
 #include "IDetailsView.h"
 #include "IDocumentation.h"
-#include "Nodes/InterchangeBaseNodeContainer.h"
-#include "Nodes/InterchangeFactoryBaseNode.h"
 #include "PropertyEditorDelegates.h"
 #include "PropertyEditorModule.h"
 #include "Styling/SlateIconFinder.h"
@@ -501,6 +504,35 @@ void SInterchangeGraphInspectorWindow::Construct(const FArguments& InArgs)
 	check(InterchangeBaseNodeContainer != nullptr);
 	check(OwnerWindow.IsValid());
 
+	{
+		TRACE_CPUPROFILER_EVENT_SCOPE(SInterchangeGraphInspectorWindow::SortingNodes)
+
+		InterchangeBaseNodeContainer->SortNodes(
+			[](const UInterchangeBaseNode& LHS, const UInterchangeBaseNode& RHS)
+			{
+				const bool bLeftIsFactory = LHS.IsA<UInterchangeFactoryBaseNode>();
+				const bool bRightIsFactory = RHS.IsA<UInterchangeFactoryBaseNode>();
+				const bool bLeftIsScene = LHS.IsA<UInterchangeSceneNode>();
+				const bool bRightIsScene = RHS.IsA<UInterchangeSceneNode>();
+
+				// Always put scene nodes above other nodes
+				if (bLeftIsScene != bRightIsScene)
+				{
+					return bLeftIsScene > bRightIsScene;
+				}
+				// Then prefer other nodes above factory nodes
+				else if (bLeftIsFactory != bRightIsFactory)
+				{
+					return bLeftIsFactory < bRightIsFactory;
+				}
+				else
+				{
+					return LHS.GetDisplayLabel() < RHS.GetDisplayLabel();
+				}
+			}
+		);
+	}
+
 	this->ChildSlot
 	[
 		SNew(SBorder)
diff --git a/Engine/Source/Runtime/Interchange/Core/Public/Nodes/InterchangeBaseNodeContainer.h b/Engine/Source/Runtime/Interchange/Core/Public/Nodes/InterchangeBaseNodeContainer.h
--- a/Engine/Source/Runtime/Interchange/Core/Public/Nodes/InterchangeBaseNodeContainer.h
+++ b/Engine/Source/Runtime/Interchange/Core/Public/Nodes/InterchangeBaseNodeContainer.h
@@ -272,6 +272,12 @@ public:
 		const FString& OldNodeUID,
 		const FString& ParentNodeUID = TEXT(""));
 
+	template<typename PREDICATE_CLASS>
+	inline void SortNodes(const PREDICATE_CLASS& Predicate)
+	{
+		Nodes.ValueSort(Predicate);
+	}
+
 private:
 
 	void InternalReorderChildren(TArray<FString>& Children) const;
