Interchange: Don't show the "There were issues with the import" if the recorded issues were due to previous imports

diff --git a/Engine/Plugins/Interchange/Editor/Source/InterchangeEditor/Private/InterchangeEditorModule.cpp b/Engine/Plugins/Interchange/Editor/Source/InterchangeEditor/Private/InterchangeEditorModule.cpp
--- a/Engine/Plugins/Interchange/Editor/Source/InterchangeEditor/Private/InterchangeEditorModule.cpp
+++ b/Engine/Plugins/Interchange/Editor/Source/InterchangeEditor/Private/InterchangeEditorModule.cpp
@@ -62,7 +62,7 @@ namespace UE::Interchange::InterchangeEditorModule
 			LogListing->SetLabel(LogListingLabel);
 
 			TArray<TSharedRef<FTokenizedMessage>> TokenizedMessages;
-			bool bHasErrors = false;
+			bool bCurrentImportHasErrors = false;
 
 			for (UInterchangeResult* Result : InResultsContainer->GetResults())
 			{
@@ -70,14 +70,18 @@ namespace UE::Interchange::InterchangeEditorModule
 				if (ResultType != EInterchangeResultType::Success)
 				{
 					TokenizedMessages.Add(FTokenizedMessage::Create(ResultType == EInterchangeResultType::Error ? EMessageSeverity::Error : EMessageSeverity::Warning, Result->GetMessageLogText()));
-					bHasErrors |= ResultType == EInterchangeResultType::Error;
+					bCurrentImportHasErrors |= ResultType == EInterchangeResultType::Error;
 				}
 			}
 
 			// UInterchangeManager::ReleaseAsyncHelper has output logs turned on. To avoid duplication, turning it off here.
 			const bool bMirrorToOutputLog = GIsAutomationTesting;
 			LogListing->AddMessages(TokenizedMessages, bMirrorToOutputLog);
-			LogListing->NotifyIfAnyMessages(NSLOCTEXT("Interchange", "LogAndNotify", "There were issues with the import."), EMessageSeverity::Info);
+
+			if (bCurrentImportHasErrors)
+			{
+				LogListing->NotifyIfAnyMessages(NSLOCTEXT("Interchange", "LogAndNotify", "There were issues with the import."), EMessageSeverity::Info);
+			}
 		}
 	}
 
