[Interchange UI] - Preview Window Allows searching for the custom added rows.

#rnx
#Synced-CL 48730524

diff --git a/Engine/Plugins/Interchange/Editor/Source/Pipelines/Private/InterchangeEditorPipelineDetails.cpp b/Engine/Plugins/Interchange/Editor/Source/Pipelines/Private/InterchangeEditorPipelineDetails.cpp
--- a/Engine/Plugins/Interchange/Editor/Source/Pipelines/Private/InterchangeEditorPipelineDetails.cpp
+++ b/Engine/Plugins/Interchange/Editor/Source/Pipelines/Private/InterchangeEditorPipelineDetails.cpp
@@ -918,24 +918,56 @@ void FInterchangeBaseNodeDetailsCustomization::CustomizeDetails( IDetailLayoutBu
 		const FText NodeInformationCategoryText = LOCTEXT("NodeInformationCategoryName", "Node Information");
 		const FName NodeInformationCategoryName = FName(TEXT("Node Information"));
 		IDetailCategoryBuilder& AttributeCategoryBuilder = DetailBuilder.EditCategory(NodeInformationCategoryName, NodeInformationCategoryText);
-		FText NodeInformationText = FText::GetEmpty();
+		
+		FText NodeTypeValueText = FText::GetEmpty();
+		FText ClassNameValueText = FText::GetEmpty();
 		if (UInterchangeFactoryBaseNode* FactoryNode = Cast<UInterchangeFactoryBaseNode>(InterchangeBaseNode))
 		{
-			const FString ClassName = FactoryNode->GetObjectClass() ? FactoryNode->GetObjectClass()->GetName() : InterchangeBaseNode->GetClass()->GetName();
-			NodeInformationText = FText::Format(LOCTEXT("NodeFactoryInformationText_FactoryBaseNode", "Factory Node ({0})"), FText::FromString(ClassName));
+			const FString ClassName = FactoryNode->GetObjectClass() 
+										? FactoryNode->GetObjectClass()->GetName() 
+										: "None - Unprocessed";
+
+			NodeTypeValueText = LOCTEXT("NodeFactoryInformation_FactoryNodeValueText", "Factory Node");
+			ClassNameValueText = FText::FromString(ClassName);
 		}
 		else
 		{
-			NodeInformationText = FText::Format(LOCTEXT("NodeFactoryInformationText_BaseNode", "Translated Node ({0})"), FText::FromString(InterchangeBaseNode->GetClass()->GetName()));
+			NodeTypeValueText = LOCTEXT("NodeFactoryInformation_TranslatedNodeValueText", "Translated Node");
+			ClassNameValueText = FText::FromString(InterchangeBaseNode->GetClass()->GetName());
 		}
 
-		FDetailWidgetRow& CustomRow = AttributeCategoryBuilder.AddCustomRow(NodeInformationCategoryText)
-		.NameContent()
-		[
-			SNew(STextBlock)
-				.Text(NodeInformationText)
-				.Font(IDetailLayoutBuilder::GetDetailFont())
-		];
+
+		const FText NodeTypeNameText = LOCTEXT("NodeFactoryInformationNodeTypeNameText", "Node Type");
+		const FText ClassNameNameText = LOCTEXT("NodeFactoryInformationClassNameNameText", "Class Name");
+		AttributeCategoryBuilder.AddCustomRow(NodeInformationCategoryText)
+			.FilterString(NodeTypeNameText)
+			.NameContent()
+			[
+				SNew(STextBlock)
+					.Text(NodeTypeNameText)
+					.Font(IDetailLayoutBuilder::GetDetailFont())
+			]
+			.ValueContent()
+			[
+				SNew(STextBlock)
+					.Text(NodeTypeValueText)
+					.Font(IDetailLayoutBuilder::GetDetailFont())
+			];
+
+		AttributeCategoryBuilder.AddCustomRow(NodeInformationCategoryText)
+			.FilterString(ClassNameNameText)
+			.NameContent()
+			[
+				SNew(STextBlock)
+					.Text(ClassNameNameText)
+					.Font(IDetailLayoutBuilder::GetDetailFont())
+			]
+			.ValueContent()
+			[
+				SNew(STextBlock)
+					.Text(ClassNameValueText)
+					.Font(IDetailLayoutBuilder::GetDetailFont())
+			];
 	}
 
 	TArray< UE::Interchange::FAttributeKey> AttributeKeys;
