[Interchange UI] Explicitly remove the possibility of Preview Panel showing the Show Property Matrix Button when relevant cvar is enabled.

#rnx
#Synced-CL 49034748

diff --git a/Engine/Plugins/Interchange/Editor/Source/Pipelines/Private/SInterchangeGraphInspectorWindow.cpp b/Engine/Plugins/Interchange/Editor/Source/Pipelines/Private/SInterchangeGraphInspectorWindow.cpp
--- a/Engine/Plugins/Interchange/Editor/Source/Pipelines/Private/SInterchangeGraphInspectorWindow.cpp
+++ b/Engine/Plugins/Interchange/Editor/Source/Pipelines/Private/SInterchangeGraphInspectorWindow.cpp
@@ -484,6 +484,9 @@ TSharedRef<SBox> SInterchangeGraphInspectorWindow::SpawnGraphInspector()
 
 	FPropertyEditorModule& PropertyEditorModule = FModuleManager::GetModuleChecked<FPropertyEditorModule>("PropertyEditor");
 	FDetailsViewArgs DetailsViewArgs;
+	// TODO: @vjavdekar - hidden as this doesn't currently provide any value and existing open modal windows block the property matrix window from opening.
+	// Reconsider this when the Preview Panel gets a UI Pass.
+	DetailsViewArgs.bShowPropertyMatrixButton = false;
 	DetailsViewArgs.bAllowSearch = bPreview;
 	DetailsViewArgs.NameAreaSettings = FDetailsViewArgs::HideNameArea;
 	GraphInspectorDetailsView = PropertyEditorModule.CreateDetailView(DetailsViewArgs);
