[Interchange] Combine Static Mesh Actor Visibility based Behavior

#rnx
#virtualized
#Synced-CL 49118104

diff --git a/Engine/Plugins/Fab/Source/Fab/Private/Importers/GenericAssetImporter.cpp b/Engine/Plugins/Fab/Source/Fab/Private/Importers/GenericAssetImporter.cpp
--- a/Engine/Plugins/Fab/Source/Fab/Private/Importers/GenericAssetImporter.cpp
+++ b/Engine/Plugins/Fab/Source/Fab/Private/Importers/GenericAssetImporter.cpp
@@ -62,7 +62,7 @@ UObject* FFabGenericImporter::GetImportOptions(const FString& SourceFile, UObjec
 					{
 						GenericAssetsPipeline->MeshPipeline->bImportStaticMeshes                  = true;
 						GenericAssetsPipeline->MeshPipeline->bImportSkeletalMeshes                = true;
-						GenericAssetsPipeline->MeshPipeline->bCombineStaticMeshes                 = true;
+						GenericAssetsPipeline->MeshPipeline->CombineStaticMeshesBehavior          = EInterchangeCombineStaticMeshesBehavior::All;
 						GenericAssetsPipeline->MeshPipeline->SkeletalMeshImportContentType        = EInterchangeSkeletalMeshContentType::All;
 						GenericAssetsPipeline->MeshPipeline->bGenerateLightmapUVs                 = true;
 						GenericAssetsPipeline->MeshPipeline->bBuildNanite                         = false;
diff --git a/Engine/Plugins/Interchange/Editor/Source/InterchangeEditor/Private/InterchangeFbxAssetImportDataConverter.cpp b/Engine/Plugins/Interchange/Editor/Source/InterchangeEditor/Private/InterchangeFbxAssetImportDataConverter.cpp
--- a/Engine/Plugins/Interchange/Editor/Source/InterchangeEditor/Private/InterchangeFbxAssetImportDataConverter.cpp
+++ b/Engine/Plugins/Interchange/Editor/Source/InterchangeEditor/Private/InterchangeFbxAssetImportDataConverter.cpp
@@ -128,7 +128,7 @@ namespace UE::Interchange::Private
 		DestinationStaticMeshImportData->bAutoGenerateCollision = GenericAssetPipeline->MeshPipeline->bCollision && GenericAssetPipeline->MeshPipeline->Collision != EInterchangeMeshCollision::None;
 		DestinationStaticMeshImportData->bBuildNanite = GenericAssetPipeline->MeshPipeline->bBuildNanite;
 		DestinationStaticMeshImportData->bBuildReversedIndexBuffer = GenericAssetPipeline->MeshPipeline->bBuildReversedIndexBuffer;
-		DestinationStaticMeshImportData->bCombineMeshes = GenericAssetPipeline->MeshPipeline->bCombineStaticMeshes;
+		DestinationStaticMeshImportData->bCombineMeshes = GenericAssetPipeline->MeshPipeline->CombineStaticMeshesBehavior == EInterchangeCombineStaticMeshesBehavior::All;
 		DestinationStaticMeshImportData->bGenerateLightmapUVs = GenericAssetPipeline->MeshPipeline->bGenerateLightmapUVs;
 		DestinationStaticMeshImportData->bOneConvexHullPerUCX = GenericAssetPipeline->MeshPipeline->bOneConvexHullPerUCX;
 		DestinationStaticMeshImportData->bRemoveDegenerates = GenericAssetPipeline->CommonMeshesProperties->bRemoveDegenerates;
@@ -317,7 +317,7 @@ namespace UE::Interchange::Private
 		GenericAssetPipeline->MeshPipeline->Collision = StaticMeshImportData->bAutoGenerateCollision ? EInterchangeMeshCollision::Convex18DOP : EInterchangeMeshCollision::None;
 		GenericAssetPipeline->MeshPipeline->bBuildNanite = StaticMeshImportData->bBuildNanite;
 		GenericAssetPipeline->MeshPipeline->bBuildReversedIndexBuffer = StaticMeshImportData->bBuildReversedIndexBuffer;
-		GenericAssetPipeline->MeshPipeline->bCombineStaticMeshes = StaticMeshImportData->bCombineMeshes;
+		GenericAssetPipeline->MeshPipeline->CombineStaticMeshesBehavior = StaticMeshImportData->bCombineMeshes ? EInterchangeCombineStaticMeshesBehavior::All : EInterchangeCombineStaticMeshesBehavior::DoNotCombine;
 		GenericAssetPipeline->MeshPipeline->bGenerateLightmapUVs = StaticMeshImportData->bGenerateLightmapUVs;
 		GenericAssetPipeline->MeshPipeline->bOneConvexHullPerUCX = StaticMeshImportData->bOneConvexHullPerUCX;
 		GenericAssetPipeline->CommonMeshesProperties->bRemoveDegenerates = StaticMeshImportData->bRemoveDegenerates;
diff --git a/Engine/Plugins/Interchange/Runtime/Source/Pipelines/Private/InterchangeGenericMeshPipeline.cpp b/Engine/Plugins/Interchange/Runtime/Source/Pipelines/Private/InterchangeGenericMeshPipeline.cpp
--- a/Engine/Plugins/Interchange/Runtime/Source/Pipelines/Private/InterchangeGenericMeshPipeline.cpp
+++ b/Engine/Plugins/Interchange/Runtime/Source/Pipelines/Private/InterchangeGenericMeshPipeline.cpp
@@ -109,7 +109,7 @@ void UInterchangeGenericMeshPipeline::AdjustSettingsForContext(const FInterchang
 			CommonMeshesProperties->VertexColorImportOption = EInterchangeVertexColorImportOption::IVCIO_Ignore;
 			bImportSkeletalMeshes = false;
 			bImportStaticMeshes = true;
-			bCombineStaticMeshes = true;
+			CombineStaticMeshesBehavior = EInterchangeCombineStaticMeshesBehavior::All;
 			bBuildNanite = false;
 			LodGroup = NAME_None;
 			bCollision = false;
@@ -132,7 +132,7 @@ void UInterchangeGenericMeshPipeline::AdjustSettingsForContext(const FInterchang
 					CommonMeshesProperties->ForceAllMeshAsType = EInterchangeForceMeshType::IFMT_StaticMesh;
 					CommonMeshesProperties->bAutoDetectMeshType = false;
 					bImportSkeletalMeshes = false;
-					bCombineStaticMeshes = true;
+					CombineStaticMeshesBehavior = EInterchangeCombineStaticMeshesBehavior::All;
 					LodGroup = NAME_None;
 					bSupportFaceRemap = false;
 					bCollision = false;
@@ -261,7 +261,7 @@ bool UInterchangeGenericMeshPipeline::IsPropertyChangeNeedRefresh(const FPropert
 		GET_MEMBER_NAME_CHECKED(UInterchangeGenericMeshPipeline, bImportStaticMeshes),
 		GET_MEMBER_NAME_CHECKED(UInterchangeGenericMeshPipeline, bImportSkeletalMeshes),
 		GET_MEMBER_NAME_CHECKED(UInterchangeGenericMeshPipeline, bCreatePhysicsAsset),
-		GET_MEMBER_NAME_CHECKED(UInterchangeGenericMeshPipeline, bCombineStaticMeshes)
+		GET_MEMBER_NAME_CHECKED(UInterchangeGenericMeshPipeline, CombineStaticMeshesBehavior),
 	};
 	
 	if (NeedRefreshProperties.Contains(PropertyChangedEvent.GetPropertyName()))
@@ -399,6 +399,15 @@ void UInterchangeGenericMeshPipeline::PostLoad()
 	{
 		bCollision = bImportCollision_DEPRECATED;
 	}
+
+	if (!bCombineStaticMeshes_DEPRECATED)
+	{
+		CombineStaticMeshesBehavior = EInterchangeCombineStaticMeshesBehavior::DoNotCombine;
+	}
+	else if (bCombineStaticMeshes_DEPRECATED)
+	{
+		CombineStaticMeshesBehavior = EInterchangeCombineStaticMeshesBehavior::All;
+	}
 }
 
 UInterchangePipelineMeshesUtilities* UInterchangeGenericMeshPipeline::CreateMeshPipelineUtilities(UInterchangeBaseNodeContainer* InBaseNodeContainer
diff --git a/Engine/Plugins/Interchange/Runtime/Source/Pipelines/Private/InterchangeGenericStaticMeshPipeline.cpp b/Engine/Plugins/Interchange/Runtime/Source/Pipelines/Private/InterchangeGenericStaticMeshPipeline.cpp
--- a/Engine/Plugins/Interchange/Runtime/Source/Pipelines/Private/InterchangeGenericStaticMeshPipeline.cpp
+++ b/Engine/Plugins/Interchange/Runtime/Source/Pipelines/Private/InterchangeGenericStaticMeshPipeline.cpp
@@ -271,9 +271,9 @@ void UInterchangeGenericMeshPipeline::ExecutePreImportPipelineStaticMesh()
 
 	if (bImportStaticMeshes && (CommonMeshesProperties->ForceAllMeshAsType == EInterchangeForceMeshType::IFMT_None || CommonMeshesProperties->ForceAllMeshAsType == EInterchangeForceMeshType::IFMT_StaticMesh))
 	{
-		if (bCombineStaticMeshes)
+		if (CombineStaticMeshesBehavior != EInterchangeCombineStaticMeshesBehavior::DoNotCombine)
 		{
-			// Combine all the static meshes
+			// Combine static meshes based on the selected combine behavior
 
 			bool bFoundMeshes = false;
 			{
@@ -294,6 +294,24 @@ void UInterchangeGenericMeshPipeline::ExecutePreImportPipelineStaticMesh()
 						TArray<FString>& TranslatedNodes = MeshUidsPerLodIndex.FindOrAdd(LodIndex);
 						for (const UInterchangeBaseNode* BaseNode : SceneNodeContainer.BaseNodes)
 						{
+							switch (CombineStaticMeshesBehavior)
+							{
+							case EInterchangeCombineStaticMeshesBehavior::VisibleOnly:
+								if (const UInterchangeSceneNode* SceneNode = Cast<const UInterchangeSceneNode>(BaseNode))
+								{
+									bool bIsActorVisible;
+									if (SceneNode->GetCustomActorVisibility(bIsActorVisible) && !bIsActorVisible)
+									{
+										continue;
+									}
+								}
+
+							case EInterchangeCombineStaticMeshesBehavior::All:
+								// Falls Through...
+							default:
+								break;
+							}
+							
 							TranslatedNodes.Add(BaseNode->GetUniqueID());
 						}
 					}
@@ -1213,4 +1231,4 @@ void UInterchangeGenericMeshPipeline::FindAndProcessIdenticalStaticMeshLODGroups
 				}
 			});
 	}
-}
+}
\ No newline at end of file
diff --git a/Engine/Plugins/Interchange/Runtime/Source/Pipelines/Public/InterchangeGenericMeshPipeline.h b/Engine/Plugins/Interchange/Runtime/Source/Pipelines/Public/InterchangeGenericMeshPipeline.h
--- a/Engine/Plugins/Interchange/Runtime/Source/Pipelines/Public/InterchangeGenericMeshPipeline.h
+++ b/Engine/Plugins/Interchange/Runtime/Source/Pipelines/Public/InterchangeGenericMeshPipeline.h
@@ -27,6 +27,16 @@ class UInterchangeStaticMeshFactoryNode;
 class UInterchangeStaticMeshLodDataNode;
 class UPhysicsAsset;
 
+UENUM(BlueprintType)
+enum class EInterchangeCombineStaticMeshesBehavior : uint8
+{
+	All UMETA(DisplayName="Combine All", Tooltip="Combine all the meshes into a single static mesh"),
+	
+	VisibleOnly UMETA(DisplayName = "Combine Visible Meshes Only", Tooltip="Skip adding the hidden meshes to the combined Static Mesh"),
+
+	DoNotCombine UMETA(DisplayName = "Do Not Combine Static Meshes", Tooltip = "All Static Meshes will be imported as individual assets")
+};
+
 /* Hide drop down will make sure the class is not showing in the class picker */
 UCLASS(MinimalAPI, BlueprintType, hidedropdown)
 class UInterchangeGenericMeshPipeline : public UInterchangePipelineBase
@@ -51,9 +61,12 @@ public:
 	UPROPERTY(EditAnywhere, BlueprintReadWrite, Category = "Static Meshes")
 	bool bImportStaticMeshes = true;
 
-	/** If enabled, all translated static mesh nodes will be imported as a single static mesh. */
-	UPROPERTY(EditAnywhere, BlueprintReadWrite, Category = "Static Meshes")
-	bool bCombineStaticMeshes = false;
+	UPROPERTY(meta=(DeprecatedProperty, DeprecationMessage = "bCombineStaticMeshes is no longer used. Use CombineStaticMeshesBehavior"))
+	bool bCombineStaticMeshes_DEPRECATED = false;
+
+	/** Specifies the behavior for combining static meshes into a single static mesh. */
+	UPROPERTY(EditAnywhere, BlueprintReadWrite, Category = "Static Meshes", meta=(DisplayName = "Combine Static Meshes"))
+	EInterchangeCombineStaticMeshesBehavior CombineStaticMeshesBehavior = EInterchangeCombineStaticMeshesBehavior::DoNotCombine;
 
 	/** The LOD group that will be assigned to this mesh. */
 	UPROPERTY(EditAnywhere, BlueprintReadWrite, Category = "Static Meshes", meta=(DisplayName = "LOD Group"), meta = (ReimportRestrict = "true"))
