[Interchange UI] Actor Visibility Toast Notification
- When invisible actors are imported in the level, additional toast notification is shown in case the game view is disabled to ensure that the visibility of the actors is accurately reflected in the viewport.

#rnx

#Synced-CL 50908799

#Synced-CL 50908799

diff --git a/Engine/Plugins/Interchange/Editor/Source/InterchangeEditor/Private/InterchangeEditorModule.cpp b/Engine/Plugins/Interchange/Editor/Source/InterchangeEditor/Private/InterchangeEditorModule.cpp
--- a/Engine/Plugins/Interchange/Editor/Source/InterchangeEditor/Private/InterchangeEditorModule.cpp
+++ b/Engine/Plugins/Interchange/Editor/Source/InterchangeEditor/Private/InterchangeEditorModule.cpp
@@ -4,6 +4,10 @@
 #include "AssetRegistry/IAssetRegistry.h"
 #include "AssetToolsModule.h"
 #include "Engine/Engine.h"
+#include "Editor/EditorEngine.h"
+#include "Framework/Notifications/NotificationManager.h"
+#include "Framework/Commands/InputBindingManager.h"
+#include "Widgets/Notifications/SNotificationList.h"
 #include "HAL/IConsoleManager.h"
 #include "IMessageLogListing.h"
 #include "InterchangeEditorLog.h"
@@ -11,6 +15,7 @@
 #include "InterchangeManager.h"
 #include "InterchangeResetContextMenuExtender.h"
 #include "InterchangePipelineSettingsCacheHandler.h"
+#include "LevelEditorSubsystem.h"
 #include "MessageLogModule.h"
 #include "Misc/App.h"
 #include "Misc/CoreDelegates.h"
@@ -18,6 +23,8 @@
 #include "PropertyEditorModule.h"
 #include "Settings/EditorLoadingSavingSettings.h"
 
+extern UNREALED_API UEditorEngine* GEditor;
+
 #define LOCTEXT_NAMESPACE "InterchangeEditorModule"
 
 DEFINE_LOG_CATEGORY(LogInterchangeEditor);
@@ -117,6 +124,70 @@ namespace UE::Interchange::InterchangeEditorModule
 			}
 		}
 	}
+
+	void ImportedObjectAvailable(TMap<int32, TArray<TObjectPtr<UObject>>> ImportedAssetsPerSourceIndex, TMap<int32, TArray<TObjectPtr<UObject>>> ImportedSceneObjectsPerSourceIndex)
+	{
+		bool bHasHiddenActors = false;
+		for (const TPair<int32, TArray<TObjectPtr<UObject>>>& SourceIndexSceneObjectsPair : ImportedSceneObjectsPerSourceIndex)
+		{
+			for (const TObjectPtr<UObject>& SceneObject : SourceIndexSceneObjectsPair.Value)
+			{
+				if (AActor* Actor = Cast<AActor>(SceneObject.Get()))
+				{
+					bHasHiddenActors |= Actor->IsHidden();
+				}
+			}
+		}
+
+		if (bHasHiddenActors)
+		{
+			if (ensure(GEditor))
+			{
+				if (ULevelEditorSubsystem* LevelEditorSubsystem = GEditor->GetEditorSubsystem<ULevelEditorSubsystem>())
+				{
+					if (!LevelEditorSubsystem->EditorGetGameView())
+					{
+						FText ShortcutText = FText::GetEmpty();
+
+						const TSharedPtr<FUICommandInfo> ToggleGameViewCommand = FInputBindingManager::Get().FindCommandInContext(TEXT("LevelViewport"), TEXT("ToggleGameView"));
+						if (ensure(ToggleGameViewCommand.IsValid()))
+						{
+							if (FText InputChordText = ToggleGameViewCommand->GetFirstValidChord()->GetInputText(); 
+								!InputChordText.IsEmpty())
+							{
+								ShortcutText = FText::FormatOrdered(LOCTEXT("GameViewShortcutText", ", Shortcut: {0}"), InputChordText);
+							}
+						}
+
+						FNotificationInfo Info(
+							FText::FormatOrdered(
+								LOCTEXT("InvisibleActorsMessage", "Interchange imported actors that are invisible that can be visualized in the level viewport using the Game View. [Path: Viewport Toolbar > Perspective > Game View.{0}]"),
+								ShortcutText
+							)
+						);
+						Info.ExpireDuration = 6.0f;
+
+						FNotificationButtonInfo EnableGameViewButtonInfo(
+							LOCTEXT("EnableGameViewButtonLabel", "Enable Game View"),
+							LOCTEXT("EnableGameViewButtonTooltipLabel", "Enables Game View to correctly reflect the visibility states of the actors in the scene."),
+							FSimpleDelegate::CreateLambda([]()
+							{
+								if (ULevelEditorSubsystem* LevelEditorSubsystem = GEditor->GetEditorSubsystem<ULevelEditorSubsystem>())
+								{
+									LevelEditorSubsystem->EditorSetGameView(true);
+								}
+							}),
+							SNotificationItem::CS_None
+						);
+						Info.ButtonDetails.Add(EnableGameViewButtonInfo);
+
+						FSlateNotificationManager::Get().AddNotification(Info);
+					}
+				}
+			}
+		}
+	}
+
 } //ns UE::Interchange::InterchangeEditorModule
 
 FInterchangeEditorModule& FInterchangeEditorModule::Get()
@@ -140,34 +211,39 @@ void FInterchangeEditorModule::StartupModule()
 
 	auto RegisterItems = [this]()
 	{
-		FDelegateHandle InterchangeEditorModuleDelegate;
-		FDelegateHandle InterchangeEditorModuleDelegateOnImportStarted;
-		FDelegateHandle InterchangeEditorModuleDelegateOnImportFinished;
-		FDelegateHandle InterchangeEditorModuleDelegateOnSanitizeName;
+		FDelegateHandle OnBatchImportCompleteHandle;
+		FDelegateHandle OnImportStartedHandle;
+		FDelegateHandle OnImportedObjectsAvailableHandle;
+		FDelegateHandle OnImportFinishedHandle;
+		FDelegateHandle OnSanitizeNameHandle;
 
 		UInterchangeManager& InterchangeManager = UInterchangeManager::GetInterchangeManager();
-		InterchangeEditorModuleDelegate = InterchangeManager.OnBatchImportComplete.AddStatic(&InterchangeEditorModule::LogErrors);
-		InterchangeEditorModuleDelegateOnImportStarted = InterchangeManager.OnImportStarted.AddStatic(&InterchangeEditorModule::ImportStarted);
-		InterchangeEditorModuleDelegateOnImportFinished = InterchangeManager.OnImportFinished.AddStatic(&InterchangeEditorModule::ImportFinished);
+		OnBatchImportCompleteHandle = InterchangeManager.OnBatchImportComplete.AddStatic(&InterchangeEditorModule::LogErrors);
+		OnImportStartedHandle = InterchangeManager.OnImportStarted.AddStatic(&InterchangeEditorModule::ImportStarted);
+		OnImportedObjectsAvailableHandle = InterchangeManager.OnImportedObjectsAvailable.AddStatic(&InterchangeEditorModule::ImportedObjectAvailable);
+		OnImportFinishedHandle = InterchangeManager.OnImportFinished.AddStatic(&InterchangeEditorModule::ImportFinished);
 		InterchangeManager.RegisterImportDataConverter(UInterchangeFbxAssetImportDataConverter::StaticClass());
 
-		InterchangeEditorModuleDelegateOnSanitizeName = InterchangeManager.OnSanitizeName.AddLambda([](FString& SanitizeName, const ESanitizeNameTypeFlags NameType)
+		OnSanitizeNameHandle= InterchangeManager.OnSanitizeName.AddLambda([](FString& SanitizeName, const ESanitizeNameTypeFlags NameType)
 			{
 				//Call the asset tools sanitize
 				IAssetTools& AssetTools = FModuleManager::LoadModuleChecked<FAssetToolsModule>("AssetTools").Get();
 				AssetTools.SanitizeName(SanitizeName);
 			});
 
-		auto UnregisterItems = [InterchangeEditorModuleDelegate
-			, InterchangeEditorModuleDelegateOnImportStarted
-			, InterchangeEditorModuleDelegateOnImportFinished
-			, InterchangeEditorModuleDelegateOnSanitizeName]()
+		auto UnregisterItems = [
+			OnBatchImportCompleteHandle,
+			OnImportStartedHandle,
+			OnImportedObjectsAvailableHandle,
+			OnImportFinishedHandle,
+			OnSanitizeNameHandle]()
 		{
 			UInterchangeManager& InterchangeManager = UInterchangeManager::GetInterchangeManager();
-			InterchangeManager.OnBatchImportComplete.Remove(InterchangeEditorModuleDelegate);
-			InterchangeManager.OnImportStarted.Remove(InterchangeEditorModuleDelegateOnImportStarted);
-			InterchangeManager.OnImportFinished.Remove(InterchangeEditorModuleDelegateOnImportFinished);
-			InterchangeManager.OnSanitizeName.Remove(InterchangeEditorModuleDelegateOnSanitizeName);
+			InterchangeManager.OnBatchImportComplete.Remove(OnBatchImportCompleteHandle);
+			InterchangeManager.OnImportStarted.Remove(OnImportStartedHandle);
+			InterchangeManager.OnImportedObjectsAvailable.Remove(OnImportedObjectsAvailableHandle);
+			InterchangeManager.OnImportFinished.Remove(OnImportFinishedHandle);
+			InterchangeManager.OnSanitizeName.Remove(OnSanitizeNameHandle);
 		};
 
 		InterchangeManager.OnPreDestroyInterchangeManager.AddLambda(UnregisterItems);
diff --git a/Engine/Plugins/Interchange/Runtime/Source/FactoryNodes/Private/InterchangeActorFactoryNode.cpp b/Engine/Plugins/Interchange/Runtime/Source/FactoryNodes/Private/InterchangeActorFactoryNode.cpp
--- a/Engine/Plugins/Interchange/Runtime/Source/FactoryNodes/Private/InterchangeActorFactoryNode.cpp
+++ b/Engine/Plugins/Interchange/Runtime/Source/FactoryNodes/Private/InterchangeActorFactoryNode.cpp
@@ -187,10 +187,6 @@ bool UInterchangeActorFactoryNode::ApplyCustomActorVisibilityToAsset(UObject* As
 		if (Actor)
 		{
 			Actor->SetActorHiddenInGame(!bVisible);
-#if WITH_EDITOR
-			// This also hides the actor on the editor viewport
-			Actor->SetIsTemporarilyHiddenInEditor(!bVisible);
-#endif	  // WITH_EDITOR
 			return true;
 		}
 	}
diff --git a/Engine/Source/Runtime/Interchange/Engine/Private/InterchangeManager.cpp b/Engine/Source/Runtime/Interchange/Engine/Private/InterchangeManager.cpp
--- a/Engine/Source/Runtime/Interchange/Engine/Private/InterchangeManager.cpp
+++ b/Engine/Source/Runtime/Interchange/Engine/Private/InterchangeManager.cpp
@@ -266,6 +266,10 @@ UE::Interchange::FImportAsyncHelper::FImportAsyncHelper()
 
 UE::Interchange::FImportAsyncHelper::~FImportAsyncHelper()
 {
+	TMap<int32, TArray<TObjectPtr<UObject>>> ExtractedAssetsPerSourceIndex = ExtractImportedObjectsPerSourceIndex(ImportedAssetsPerSourceIndex);
+	TMap<int32, TArray<TObjectPtr<UObject>>> ExtractedSceneObjectsPerSourceIndex  = ExtractImportedObjectsPerSourceIndex(ImportedSceneObjectsPerSourceIndex);
+	UInterchangeManager::GetInterchangeManager().OnImportedObjectsAvailable.Broadcast(ExtractedAssetsPerSourceIndex, ExtractedSceneObjectsPerSourceIndex);
+
 	CleanUp();
 	
 	if (AnalyticsHandler.IsValid())
@@ -614,6 +618,21 @@ void UE::Interchange::FImportAsyncHelper::CleanUp()
 	CreatedFactories.Empty();
 }
 
+TMap<int32, TArray<TObjectPtr<UObject>>> UE::Interchange::FImportAsyncHelper::ExtractImportedObjectsPerSourceIndex(TMap<int32, TArray<FImportedObjectInfo>>& InObjectInfoPerSourceIndex)
+{
+	TMap<int32, TArray<TObjectPtr<UObject>>> ExtractedObjectsPerSourceIndex;
+	for (const TPair<int32, TArray<FImportedObjectInfo>>& ImportedAssetPerSourceIndex : InObjectInfoPerSourceIndex)
+	{
+		TArray<TObjectPtr<UObject>>& ImportedObjects = ExtractedObjectsPerSourceIndex.FindOrAdd(ImportedAssetPerSourceIndex.Key);
+		for (const FImportedObjectInfo& ImportObjectInfo : ImportedAssetPerSourceIndex.Value)
+		{
+			ImportedObjects.Add(ImportObjectInfo.GetImportedObject());
+		}
+	}
+
+	return ExtractedObjectsPerSourceIndex;
+}
+
 UE::Interchange::FImportResult::FImportResult()
 	: ImportStatus(EStatus::Invalid)
 {
diff --git a/Engine/Source/Runtime/Interchange/Engine/Public/InterchangeManager.h b/Engine/Source/Runtime/Interchange/Engine/Public/InterchangeManager.h
--- a/Engine/Source/Runtime/Interchange/Engine/Public/InterchangeManager.h
+++ b/Engine/Source/Runtime/Interchange/Engine/Public/InterchangeManager.h
@@ -351,6 +351,8 @@ namespace UE
 			// This class owns the Analytics Handler and manages its lifetime.
 			// Analytics Handler will be released at the end of import when the ImportAsyncHandler is deleted.
 			TStrongObjectPtr<UInterchangeAnalyticsHandlerBase> AnalyticsHandler;
+		private:
+				static TMap<int32, TArray<TObjectPtr<UObject>>> ExtractImportedObjectsPerSourceIndex(TMap<int32, TArray<FImportedObjectInfo>>& InObjectInfoPerSourceIndex);
 
 		private:
 			FCriticalSection ClassPermissionLock;
@@ -505,14 +507,22 @@ public:
 	DECLARE_MULTICAST_DELEGATE_OneParam(FInterchangeOnAssetPostReimport, UObject*);
 	/** Delegate type that is fired when the import results in an error. */
 	DECLARE_MULTICAST_DELEGATE_OneParam(FInterchangeOnBatchImportComplete, TStrongObjectPtr<UInterchangeResultsContainer>);
+	/** Delegate type that is fired when the import is completed and results are available to iterate on. */
+	using FSourceIndexObjectsMap = TMap<int32, TArray<TObjectPtr<UObject>>>;
+	DECLARE_MULTICAST_DELEGATE_TwoParams(FInterchangeOnImportedObjectsAvailable, FSourceIndexObjectsMap, FSourceIndexObjectsMap);
 
 	// Delegates used to register and unregister.
 
 	FInterchangeOnAssetPostImport OnAssetPostImport;
 	FInterchangeOnAssetPostReimport OnAssetPostReimport;
 	FInterchangeOnBatchImportComplete OnBatchImportComplete;
+
 	//Fires when the first import process starts.
 	FOnImportStarted OnImportStarted;
+
+	//Fires when the imported assets and scene actors are available.
+	FInterchangeOnImportedObjectsAvailable OnImportedObjectsAvailable;
+
 	//Fires when the last import process finishes.
 	FOnImportFinished OnImportFinished;
 
