[Interchange Visibility Toast] Change Visibility of the buttons dynamically so that the buttons can be clicked only once.

#rnx
#Synced-CL 51675681

diff --git a/Engine/Plugins/Interchange/Editor/Source/InterchangeEditor/Private/InterchangeEditorModule.cpp b/Engine/Plugins/Interchange/Editor/Source/InterchangeEditor/Private/InterchangeEditorModule.cpp
--- a/Engine/Plugins/Interchange/Editor/Source/InterchangeEditor/Private/InterchangeEditorModule.cpp
+++ b/Engine/Plugins/Interchange/Editor/Source/InterchangeEditor/Private/InterchangeEditorModule.cpp
@@ -164,12 +164,25 @@ namespace UE::Interchange::InterchangeEditorModule
 							}
 						}
 
-						FNotificationInfo Info(
-							FText::FormatOrdered(
+						FNotificationInfo Info(FText::GetEmpty());
+						Info.Text = TAttribute<FText>::CreateLambda([ShortcutText]()
+						{
+							if (GEditor)
+							{
+								if (ULevelEditorSubsystem* LevelEditorSubsystem = GEditor->GetEditorSubsystem<ULevelEditorSubsystem>())
+								{
+									if (LevelEditorSubsystem->EditorGetGameView())
+									{
+										return LOCTEXT("GameViewEnabledLabel", "Game View Enabled");
+									}
+								}
+							}
+
+							return FText::FormatOrdered(
 								LOCTEXT("InvisibleActorsMessage", "Interchange imported actors that are invisible that can be visualized in the level viewport using the Game View. [Path: Viewport Toolbar > Perspective > Game View.{0}]"),
 								ShortcutText
-							)
-						);
+							);
+						});
 						Info.ExpireDuration = 6.0f;
 
 						FNotificationButtonInfo EnableGameViewButtonInfo(
@@ -177,15 +190,38 @@ namespace UE::Interchange::InterchangeEditorModule
 							LOCTEXT("EnableGameViewButtonTooltipLabel", "Enables Game View to correctly reflect the visibility states of the actors in the scene."),
 							FSimpleDelegate::CreateLambda([]()
 							{
-								if (ULevelEditorSubsystem* LevelEditorSubsystem = GEditor->GetEditorSubsystem<ULevelEditorSubsystem>())
+								if (GEditor)
+								{
+									if (ULevelEditorSubsystem* LevelEditorSubsystem = GEditor->GetEditorSubsystem<ULevelEditorSubsystem>())
+									{
+										LevelEditorSubsystem->EditorSetGameView(true);
+									}
+								}
+							}),
+							FNotificationButtonInfo::FVisibilityDelegate::CreateLambda([](SNotificationItem::ECompletionState)->EVisibility
+							{
+								if (GEditor)
 								{
-									LevelEditorSubsystem->EditorSetGameView(true);
+									if (ULevelEditorSubsystem* LevelEditorSubsystem = GEditor->GetEditorSubsystem<ULevelEditorSubsystem>())
+									{
+										return	(!LevelEditorSubsystem->EditorGetGameView() ? EVisibility::Visible : EVisibility::Collapsed);
+									}
 								}
+								return EVisibility::Visible;
 							}),
-							SNotificationItem::CS_None
+							FNotificationButtonInfo::FIsEnabledDelegate::CreateLambda([](SNotificationItem::ECompletionState)->bool
+							{
+								if(GEditor)
+								{ 
+									if (ULevelEditorSubsystem* LevelEditorSubsystem = GEditor->GetEditorSubsystem<ULevelEditorSubsystem>())
+									{
+										return !LevelEditorSubsystem->EditorGetGameView();
+									}
+								}
+								return false;
+							})
 						);
 						Info.ButtonDetails.Add(EnableGameViewButtonInfo);
-
 						FSlateNotificationManager::Get().AddNotification(Info);
 					}
 				}
