[Interchange Reimport] FIX Material Reimport - Asset Registry Warnings
- Optimized the Interchange Cards
- Added ability to get base content path for reimport while executing preview pipelines.

#rnx
#Synced-CL 52017234

diff --git a/Engine/Plugins/Interchange/Editor/Source/Pipelines/Private/SInterchangePipelineConfigurationDialog.cpp b/Engine/Plugins/Interchange/Editor/Source/Pipelines/Private/SInterchangePipelineConfigurationDialog.cpp
--- a/Engine/Plugins/Interchange/Editor/Source/Pipelines/Private/SInterchangePipelineConfigurationDialog.cpp
+++ b/Engine/Plugins/Interchange/Editor/Source/Pipelines/Private/SInterchangePipelineConfigurationDialog.cpp
@@ -1774,6 +1774,11 @@ void SInterchangePipelineConfigurationDialog::FillAssetCardsList()
 
 void SInterchangePipelineConfigurationDialog::CreateCardsViewList()
 {
+	if (!bShowCards)
+	{
+		return;
+	}
+
 	FillAssetCardsList();
 	if (AssetCards.IsEmpty())
 	{
@@ -1787,6 +1792,11 @@ void SInterchangePipelineConfigurationDialog::CreateCardsViewList()
 
 void SInterchangePipelineConfigurationDialog::RefreshCardsViewList()
 {
+	if (!bShowCards)
+	{
+		return;
+	}
+
 	FillAssetCardsList();
 	if (CardViewList.IsValid())
 	{
@@ -1835,6 +1845,12 @@ void SInterchangePipelineConfigurationDialog::UpdatePreviewContainer(bool bUpdat
 	TArray<UInterchangeSourceData*> SourceDatas;
 	SourceDatas.Add(SourceData.Get());
 
+	FString ContentBasePath = FString();
+	if (bReimport && ReimportObject.IsValid())
+	{
+		ContentBasePath = FPaths::GetPath(ReimportObject->GetPathName());
+	}
+
 	//Execute all pipelines on the duplicated container
 	UInterchangeResultsContainer* Results = NewObject<UInterchangeResultsContainer>(GetTransientPackage());
 	for (int32 PipelineIndex = 0; PipelineIndex < PipelineListViewItems.Num(); ++PipelineIndex)
@@ -1846,7 +1862,7 @@ void SInterchangePipelineConfigurationDialog::UpdatePreviewContainer(bool bUpdat
 		{
 			DuplicatedPipeline->TransferAdjustSettings(PipelineItem->Pipeline);
 			DuplicatedPipeline->SetResultsContainer(Results);
-			DuplicatedPipeline->ScriptedExecutePipeline(PreviewNodeContainer.Get(), SourceDatas, FString());
+			DuplicatedPipeline->ScriptedExecutePipeline(PreviewNodeContainer.Get(), SourceDatas, ContentBasePath);
 			ClearObjectFlags(DuplicatedPipeline);
 		}
 	}
@@ -1856,7 +1872,7 @@ void SInterchangePipelineConfigurationDialog::UpdatePreviewContainer(bool bUpdat
 		//If we do not update cards execute the cards pipeline since its a final preview
 		if (UInterchangeCardsPipeline* InterchangeCardsPipeline = GenerateTransientCardsPipeline())
 		{
-			InterchangeCardsPipeline->ScriptedExecutePipeline(PreviewNodeContainer.Get(), SourceDatas, FString());
+			InterchangeCardsPipeline->ScriptedExecutePipeline(PreviewNodeContainer.Get(), SourceDatas, ContentBasePath);
 		}
 	}
 
