[ContentBrowser] New Add Menu Interchange Menu

#Synced-CL 52417896

diff --git a/Engine/Plugins/Interchange/Editor/Source/Pipelines/InterchangeEditorPipelines.Build.cs b/Engine/Plugins/Interchange/Editor/Source/Pipelines/InterchangeEditorPipelines.Build.cs
--- a/Engine/Plugins/Interchange/Editor/Source/Pipelines/InterchangeEditorPipelines.Build.cs
+++ b/Engine/Plugins/Interchange/Editor/Source/Pipelines/InterchangeEditorPipelines.Build.cs
@@ -20,7 +20,9 @@ namespace UnrealBuildTool.Rules
 			PrivateDependencyModuleNames.AddRange(
 				new string[]
 				{
+					"AssetDefinition",
 					"AssetTools",
+					"EngineAssetDefinitions",
 					"InputCore",
 					"InterchangePipelines",
 					"InterchangeNodes",
diff --git a/Engine/Plugins/Interchange/Editor/Source/Pipelines/Private/AssetDefinition/AssetDefinition_InterchangeBlueprintPipeline.cpp b/Engine/Plugins/Interchange/Editor/Source/Pipelines/Private/AssetDefinition/AssetDefinition_InterchangeBlueprintPipeline.cpp
new file mode 100644
--- /dev/null
+++ b/Engine/Plugins/Interchange/Editor/Source/Pipelines/Private/AssetDefinition/AssetDefinition_InterchangeBlueprintPipeline.cpp
@@ -0,0 +1,15 @@
+// Copyright Epic Games, Inc. All Rights Reserved. 
+
+#include "AssetDefinition_InterchangeBlueprintPipeline.h"
+#include "InterchangePipelineFactories.h"
+
+#define LOCTEXT_NAMESPACE "AssetDefinition_InterchangeBlueprintPipeline"
+
+UFactory* UAssetDefinition_InterchangeBlueprintPipeline::GetFactoryForBlueprintType(UBlueprint* InBlueprint) const
+{
+	UInterchangeBlueprintPipelineBaseFactory* InterchangeBlueprintPipelineBaseFactory = NewObject<UInterchangeBlueprintPipelineBaseFactory>();
+	InterchangeBlueprintPipelineBaseFactory->ParentClass = TSubclassOf<UInterchangePipelineBase>(*InBlueprint->GeneratedClass);
+	return InterchangeBlueprintPipelineBaseFactory;
+}
+
+#undef LOCTEXT_NAMESPACE
diff --git a/Engine/Plugins/Interchange/Editor/Source/Pipelines/Private/AssetDefinition/AssetDefinition_InterchangeBlueprintPipeline.h b/Engine/Plugins/Interchange/Editor/Source/Pipelines/Private/AssetDefinition/AssetDefinition_InterchangeBlueprintPipeline.h
new file mode 100644
--- /dev/null
+++ b/Engine/Plugins/Interchange/Editor/Source/Pipelines/Private/AssetDefinition/AssetDefinition_InterchangeBlueprintPipeline.h
@@ -0,0 +1,25 @@
+// Copyright Epic Games, Inc. All Rights Reserved. 
+
+#pragma once
+#include "InterchangeBlueprintPipelineBase.h"
+#include "Script/AssetDefinition_Blueprint.h"
+#include "AssetDefinition_InterchangeBlueprintPipeline.generated.h"
+
+UCLASS()
+class UAssetDefinition_InterchangeBlueprintPipeline : public UAssetDefinition_Blueprint
+{
+	GENERATED_BODY()
+public:
+	// UAssetDefinition Begin
+	virtual FText GetAssetDisplayName() const override { return NSLOCTEXT("AssetDefinition", "AssetDefinition_InterchangeBlueprintPipeline", "Interchange Blueprint Pipeline"); }
+	virtual FLinearColor GetAssetColor() const override { return FColor(10, 25, 175); }
+	virtual TSoftClassPtr<UObject> GetAssetClass() const override { return UInterchangeBlueprintPipelineBase::StaticClass(); }
+	virtual TConstArrayView<FAssetCategoryPath> GetAssetCategories() const override
+	{
+		static const FAssetCategoryPath Interchange(NSLOCTEXT("InterchangeBlueprintPipeline", "InterchangeCategoryPath", "Interchange"));
+		static const auto Categories = { FAssetCategoryPath(Interchange, NSLOCTEXT("AssetDefinition", "AssetDefinition_InterchangeBlueprintPipelineSubMenu", "Interchange Blueprint"), ECategoryMenuType::Section) };
+		return Categories;
+	}
+	virtual UFactory* GetFactoryForBlueprintType(UBlueprint* InBlueprint) const override;
+	// UAssetDefinition End
+};
diff --git a/Engine/Plugins/Interchange/Editor/Source/Pipelines/Private/AssetDefinition/AssetDefinition_InterchangeEditorBlueprintPipeline.cpp b/Engine/Plugins/Interchange/Editor/Source/Pipelines/Private/AssetDefinition/AssetDefinition_InterchangeEditorBlueprintPipeline.cpp
new file mode 100644
--- /dev/null
+++ b/Engine/Plugins/Interchange/Editor/Source/Pipelines/Private/AssetDefinition/AssetDefinition_InterchangeEditorBlueprintPipeline.cpp
@@ -0,0 +1,15 @@
+// Copyright Epic Games, Inc. All Rights Reserved. 
+
+#include "AssetDefinition_InterchangeEditorBlueprintPipeline.h"
+#include "InterchangePipelineFactories.h"
+
+#define LOCTEXT_NAMESPACE "AssetDefinition_InterchangeEditorBlueprintPipeline"
+
+UFactory* UAssetDefinition_InterchangeEditorBlueprintPipeline::GetFactoryForBlueprintType(UBlueprint* InBlueprint) const
+{
+	UInterchangeEditorBlueprintPipelineBaseFactory* InterchangeEditorBlueprintPipelineBaseFactory = NewObject<UInterchangeEditorBlueprintPipelineBaseFactory>();
+	InterchangeEditorBlueprintPipelineBaseFactory->ParentClass = TSubclassOf<UInterchangeEditorPipelineBase>(*InBlueprint->GeneratedClass);
+	return InterchangeEditorBlueprintPipelineBaseFactory;
+}
+
+#undef LOCTEXT_NAMESPACE
diff --git a/Engine/Plugins/Interchange/Editor/Source/Pipelines/Private/AssetDefinition/AssetDefinition_InterchangeEditorBlueprintPipeline.h b/Engine/Plugins/Interchange/Editor/Source/Pipelines/Private/AssetDefinition/AssetDefinition_InterchangeEditorBlueprintPipeline.h
new file mode 100644
--- /dev/null
+++ b/Engine/Plugins/Interchange/Editor/Source/Pipelines/Private/AssetDefinition/AssetDefinition_InterchangeEditorBlueprintPipeline.h
@@ -0,0 +1,25 @@
+// Copyright Epic Games, Inc. All Rights Reserved. 
+
+#pragma once
+#include "InterchangeEditorBlueprintPipelineBase.h"
+#include "Script/AssetDefinition_Blueprint.h"
+#include "AssetDefinition_InterchangeEditorBlueprintPipeline.generated.h"
+
+UCLASS()
+class UAssetDefinition_InterchangeEditorBlueprintPipeline : public UAssetDefinition_Blueprint
+{
+	GENERATED_BODY()
+public:
+	// UAssetDefinition Begin
+	virtual FText GetAssetDisplayName() const override { return NSLOCTEXT("AssetDefinition", "AssetDefinition_InterchangeEditorBlueprintPipeline", "Interchange Editor Blueprint Pipeline"); }
+	virtual FLinearColor GetAssetColor() const override { return FColor(10, 25, 175); }
+	virtual TSoftClassPtr<UObject> GetAssetClass() const override { return UInterchangeEditorBlueprintPipelineBase::StaticClass(); }
+	virtual TConstArrayView<FAssetCategoryPath> GetAssetCategories() const override
+	{
+		static const FAssetCategoryPath Interchange(NSLOCTEXT("InterchangeEditorBlueprintPipeline", "InterchangeCategoryPath", "Interchange"));
+		static const auto Categories = { FAssetCategoryPath(Interchange, NSLOCTEXT("AssetDefinition", "AssetDefinition_InterchangeEditorBlueprintPipelineSubMenu", "Interchange Blueprint"), ECategoryMenuType::Section) };
+		return Categories;
+	}
+	virtual UFactory* GetFactoryForBlueprintType(UBlueprint* InBlueprint) const override;
+	// UAssetDefinition End
+};
diff --git a/Engine/Plugins/Interchange/Editor/Source/Pipelines/Private/AssetDefinition/AssetDefinition_InterchangePipeline.cpp b/Engine/Plugins/Interchange/Editor/Source/Pipelines/Private/AssetDefinition/AssetDefinition_InterchangePipeline.cpp
new file mode 100644
--- /dev/null
+++ b/Engine/Plugins/Interchange/Editor/Source/Pipelines/Private/AssetDefinition/AssetDefinition_InterchangePipeline.cpp
@@ -0,0 +1,18 @@
+// Copyright Epic Games, Inc. All Rights Reserved. 
+
+#include "AssetDefinition_InterchangePipeline.h"
+#include "Toolkits/SimpleAssetEditor.h"
+
+#define LOCTEXT_NAMESPACE "AssetDefinition_InterchangePipeline"
+
+EAssetCommandResult UAssetDefinition_InterchangePipeline::OpenAssets(const FAssetOpenArgs& OpenArgs) const
+{
+	for (UInterchangePipelineBase* Object : OpenArgs.LoadObjects<UInterchangePipelineBase>())
+	{
+		FSimpleAssetEditor::CreateEditor(EToolkitMode::Standalone, OpenArgs.ToolkitHost, Object);
+	}
+
+	return EAssetCommandResult::Handled;
+}
+
+#undef LOCTEXT_NAMESPACE
diff --git a/Engine/Plugins/Interchange/Editor/Source/Pipelines/Private/AssetDefinition/AssetDefinition_InterchangePipeline.h b/Engine/Plugins/Interchange/Editor/Source/Pipelines/Private/AssetDefinition/AssetDefinition_InterchangePipeline.h
new file mode 100644
--- /dev/null
+++ b/Engine/Plugins/Interchange/Editor/Source/Pipelines/Private/AssetDefinition/AssetDefinition_InterchangePipeline.h
@@ -0,0 +1,25 @@
+// Copyright Epic Games, Inc. All Rights Reserved. 
+
+#pragma once
+#include "AssetDefinitionDefault.h"
+#include "InterchangePipelineBase.h"
+#include "AssetDefinition_InterchangePipeline.generated.h"
+
+UCLASS()
+class UAssetDefinition_InterchangePipeline : public UAssetDefinitionDefault
+{
+	GENERATED_BODY()
+public:
+	// UAssetDefinition Begin
+	virtual FText GetAssetDisplayName() const override { return NSLOCTEXT("AssetDefinition", "AssetDefinition_InterchangePipeline", "Interchange Pipeline"); }
+	virtual FLinearColor GetAssetColor() const override { return FColor(135, 200, 25); }
+	virtual TSoftClassPtr<UObject> GetAssetClass() const override { return UInterchangePipelineBase::StaticClass(); }
+	virtual TConstArrayView<FAssetCategoryPath> GetAssetCategories() const override
+	{
+		static const FAssetCategoryPath Interchange(NSLOCTEXT("InterchangePipeline", "InterchangeCategoryPath", "Interchange"));
+		static const auto Categories = { FAssetCategoryPath(Interchange, NSLOCTEXT("AssetDefinition", "AssetDefinition_InterchangePipelineSubMenu", "Interchange Pipeline"), ECategoryMenuType::Section) };
+		return Categories;
+	}
+	virtual EAssetCommandResult OpenAssets(const FAssetOpenArgs& OpenArgs) const override;
+	// UAssetDefinition End
+};
diff --git a/Engine/Plugins/Interchange/Editor/Source/Pipelines/Private/AssetDefinition/AssetDefinition_InterchangePythonPipeline.cpp b/Engine/Plugins/Interchange/Editor/Source/Pipelines/Private/AssetDefinition/AssetDefinition_InterchangePythonPipeline.cpp
new file mode 100644
--- /dev/null
+++ b/Engine/Plugins/Interchange/Editor/Source/Pipelines/Private/AssetDefinition/AssetDefinition_InterchangePythonPipeline.cpp
@@ -0,0 +1,18 @@
+// Copyright Epic Games, Inc. All Rights Reserved. 
+
+#include "AssetDefinition_InterchangePythonPipeline.h"
+#include "Toolkits/SimpleAssetEditor.h"
+
+#define LOCTEXT_NAMESPACE "AssetDefinition_InterchangePythonPipeline"
+
+EAssetCommandResult UAssetDefinition_InterchangePythonPipeline::OpenAssets(const FAssetOpenArgs& OpenArgs) const
+{
+	for (UInterchangePythonPipelineAsset* Object : OpenArgs.LoadObjects<UInterchangePythonPipelineAsset>())
+	{
+		FSimpleAssetEditor::CreateEditor(EToolkitMode::Standalone, OpenArgs.ToolkitHost, Object);
+	}
+
+	return EAssetCommandResult::Handled;
+}
+
+#undef LOCTEXT_NAMESPACE
diff --git a/Engine/Plugins/Interchange/Editor/Source/Pipelines/Private/AssetDefinition/AssetDefinition_InterchangePythonPipeline.h b/Engine/Plugins/Interchange/Editor/Source/Pipelines/Private/AssetDefinition/AssetDefinition_InterchangePythonPipeline.h
new file mode 100644
--- /dev/null
+++ b/Engine/Plugins/Interchange/Editor/Source/Pipelines/Private/AssetDefinition/AssetDefinition_InterchangePythonPipeline.h
@@ -0,0 +1,25 @@
+// Copyright Epic Games, Inc. All Rights Reserved. 
+
+#pragma once
+#include "AssetDefinitionDefault.h"
+#include "InterchangePythonPipelineBase.h"
+#include "AssetDefinition_InterchangePythonPipeline.generated.h"
+
+UCLASS()
+class UAssetDefinition_InterchangePythonPipeline : public UAssetDefinitionDefault
+{
+	GENERATED_BODY()
+public:
+	// UAssetDefinition Begin
+	virtual FText GetAssetDisplayName() const override { return NSLOCTEXT("AssetDefinition", "AssetDefinition_InterchangePythonPipeline", "Interchange Python Pipeline"); }
+	virtual FLinearColor GetAssetColor() const override { return FColor(135, 200, 25); }
+	virtual TSoftClassPtr<UObject> GetAssetClass() const override { return UInterchangePythonPipelineAsset::StaticClass(); }
+	virtual TConstArrayView<FAssetCategoryPath> GetAssetCategories() const override
+	{
+		static const FAssetCategoryPath Interchange(NSLOCTEXT("InterchangePythonPipeline", "InterchangeCategoryPath", "Interchange"));
+		static const auto Categories = { FAssetCategoryPath(Interchange, NSLOCTEXT("AssetDefinition", "AssetDefinition_InterchangePythonPipelineSubMenu", "Interchange Pipeline"), ECategoryMenuType::Section) };
+		return Categories;
+	}
+	virtual EAssetCommandResult OpenAssets(const FAssetOpenArgs& OpenArgs) const override;
+	// UAssetDefinition End
+};
diff --git a/Engine/Plugins/Interchange/Editor/Source/Pipelines/Private/InterchangeEditorPipelinesModule.cpp b/Engine/Plugins/Interchange/Editor/Source/Pipelines/Private/InterchangeEditorPipelinesModule.cpp
--- a/Engine/Plugins/Interchange/Editor/Source/Pipelines/Private/InterchangeEditorPipelinesModule.cpp
+++ b/Engine/Plugins/Interchange/Editor/Source/Pipelines/Private/InterchangeEditorPipelinesModule.cpp
@@ -44,6 +44,8 @@ private:
 
 	EAssetTypeCategories::Type InterchangeAssetCategory;
 
+	TArray<FName> ClassesToUnregisterOnShutdown;
+
 	/** Pointer to the style set to use for the UI. */
 	TSharedPtr<ISlateStyle> InterchangeEditorPipelineStyle = nullptr;
 
@@ -112,18 +114,6 @@ void FInterchangeEditorPipelinesModule::AcquireResources()
 
 	InterchangeAssetCategory = AssetTools.RegisterAdvancedAssetCategory(FName(TEXT("Interchange")), NSLOCTEXT("InterchangeEditorPipelineModule", "InterchangeAssetCategoryCategoryName", "Interchange"));
 	
-	PipelineBase_TypeActions = MakeShared<FAssetTypeActions_InterchangePipelineBase>(InterchangeAssetCategory);
-	AssetTools.RegisterAssetTypeActions(PipelineBase_TypeActions.ToSharedRef());
-
-	BlueprintPipelineBase_TypeActions = MakeShared<FAssetTypeActions_InterchangeBlueprintPipelineBase>(InterchangeAssetCategory);
-	AssetTools.RegisterAssetTypeActions(BlueprintPipelineBase_TypeActions.ToSharedRef());
-
-	BlueprintEditorPipelineBase_TypeActions = MakeShared<FAssetTypeActions_InterchangeEditorBlueprintPipelineBase>(InterchangeAssetCategory);
-	AssetTools.RegisterAssetTypeActions(BlueprintEditorPipelineBase_TypeActions.ToSharedRef());
-
-	PythonPipelineBase_TypeActions = MakeShared<FAssetTypeActions_InterchangePythonPipelineBase>(InterchangeAssetCategory);
-	AssetTools.RegisterAssetTypeActions(PythonPipelineBase_TypeActions.ToSharedRef());
-
 	FCoreDelegates::OnPreExit.AddLambda([this]()
 		{
 			//We must release the resources before the application start to unload modules
@@ -152,16 +142,6 @@ void FInterchangeEditorPipelinesModule::ReleaseResources()
 	UnregisterPropertySectionMappings();
 
 	InterchangeEditorPipelineStyle = nullptr;
-
-	FAssetToolsModule* AssetToolsModule = FModuleManager::GetModulePtr<FAssetToolsModule>("AssetTools");
-	if (AssetToolsModule)
-	{
-		IAssetTools& AssetTools = AssetToolsModule->Get();
-		AssetTools.UnregisterAssetTypeActions(BlueprintPipelineBase_TypeActions.ToSharedRef());
-		AssetTools.UnregisterAssetTypeActions(BlueprintEditorPipelineBase_TypeActions.ToSharedRef());
-		AssetTools.UnregisterAssetTypeActions(PipelineBase_TypeActions.ToSharedRef());
-		AssetTools.UnregisterAssetTypeActions(PythonPipelineBase_TypeActions.ToSharedRef());
-	}
 }
 
 TSharedRef<FPropertySection> FInterchangeEditorPipelinesModule::RegisterPropertySection(FPropertyEditorModule& PropertyModule, FName ClassName, FName SectionName, FText DisplayName)
diff --git a/Engine/Plugins/Interchange/Editor/Source/Pipelines/Private/InterchangePipelineFactories.cpp b/Engine/Plugins/Interchange/Editor/Source/Pipelines/Private/InterchangePipelineFactories.cpp
--- a/Engine/Plugins/Interchange/Editor/Source/Pipelines/Private/InterchangePipelineFactories.cpp
+++ b/Engine/Plugins/Interchange/Editor/Source/Pipelines/Private/InterchangePipelineFactories.cpp
@@ -61,18 +61,6 @@ UObject* UInterchangeBlueprintPipelineBaseFactory::FactoryCreateNew(UClass* Clas
 	return FactoryCreateNew(Class, InParent, Name, Flags, Context, Warn, NAME_None);
 }
 
-UFactory* FAssetTypeActions_InterchangeBlueprintPipelineBase::GetFactoryForBlueprintType(UBlueprint* InBlueprint) const
-{
-	UInterchangeBlueprintPipelineBaseFactory* InterchangeBlueprintPipelineBaseFactory = NewObject<UInterchangeBlueprintPipelineBaseFactory>();
-	InterchangeBlueprintPipelineBaseFactory->ParentClass = TSubclassOf<UInterchangePipelineBase>(*InBlueprint->GeneratedClass);
-	return InterchangeBlueprintPipelineBaseFactory;
-}
-
-UClass* FAssetTypeActions_InterchangeBlueprintPipelineBase::GetSupportedClass() const
-{
-	return UInterchangeBlueprintPipelineBase::StaticClass();
-}
-
 /**
  * UInterchangeEditorBlueprintPipelineBaseFactory implementation
  */
@@ -114,18 +102,6 @@ UObject* UInterchangeEditorBlueprintPipelineBaseFactory::FactoryCreateNew(UClass
 	return FactoryCreateNew(Class, InParent, Name, Flags, Context, Warn, NAME_None);
 }
 
-UFactory* FAssetTypeActions_InterchangeEditorBlueprintPipelineBase::GetFactoryForBlueprintType(UBlueprint* InBlueprint) const
-{
-	UInterchangeEditorBlueprintPipelineBaseFactory* InterchangeEditorBlueprintPipelineBaseFactory = NewObject<UInterchangeEditorBlueprintPipelineBaseFactory>();
-	InterchangeEditorBlueprintPipelineBaseFactory->ParentClass = TSubclassOf<UInterchangeEditorPipelineBase>(*InBlueprint->GeneratedClass);
-	return InterchangeEditorBlueprintPipelineBaseFactory;
-}
-
-UClass* FAssetTypeActions_InterchangeEditorBlueprintPipelineBase::GetSupportedClass() const
-{
-	return UInterchangeEditorBlueprintPipelineBase::StaticClass();
-}
-
 /**
  * UInterchangePipelineBaseFactory implementation
  */
@@ -181,22 +157,6 @@ bool UInterchangePipelineBaseFactory::ConfigureProperties()
 	return bPressedOk;
 }
 
-UClass* FAssetTypeActions_InterchangePipelineBase::GetSupportedClass() const
-{
-	return UInterchangePipelineBase::StaticClass();
-}
-
-void FAssetTypeActions_InterchangePipelineBase::OpenAssetEditor(const TArray<UObject*>& InObjects, TSharedPtr<class IToolkitHost> EditWithinLevelEditor)
-{
-	for (UObject* Object : InObjects)
-	{
-		if (UInterchangePipelineBase* Asset = Cast<UInterchangePipelineBase>(Object))
-		{
-			FSimpleAssetEditor::CreateEditor(EToolkitMode::Standalone, EditWithinLevelEditor, Object);
-		}
-	}
-}
-
 /**
  * UInterchangePythonPipelineFactory implementation
  */
@@ -258,19 +218,3 @@ bool UInterchangePythonPipelineAssetFactory::ConfigureProperties()
 
 	return bPressedOk;
 }
-
-UClass* FAssetTypeActions_InterchangePythonPipelineBase::GetSupportedClass() const
-{
-	return UInterchangePythonPipelineAsset::StaticClass();
-}
-
-void FAssetTypeActions_InterchangePythonPipelineBase::OpenAssetEditor(const TArray<UObject*>& InObjects, TSharedPtr<class IToolkitHost> EditWithinLevelEditor)
-{
-	for (UObject* Object : InObjects)
-	{
-		if (UInterchangePythonPipelineAsset* Asset = Cast<UInterchangePythonPipelineAsset>(Object))
-		{
-			FSimpleAssetEditor::CreateEditor(EToolkitMode::Standalone, EditWithinLevelEditor, Object);
-		}
-	}
-}
diff --git a/Engine/Plugins/Interchange/Editor/Source/Pipelines/Private/InterchangePipelineFactories.h b/Engine/Plugins/Interchange/Editor/Source/Pipelines/Private/InterchangePipelineFactories.h
--- a/Engine/Plugins/Interchange/Editor/Source/Pipelines/Private/InterchangePipelineFactories.h
+++ b/Engine/Plugins/Interchange/Editor/Source/Pipelines/Private/InterchangePipelineFactories.h
@@ -59,25 +59,6 @@ public:
 	//End UFactory Interface	
 };
 
-class FAssetTypeActions_InterchangeBlueprintPipelineBase : public FAssetTypeActions_Blueprint
-{
-public:
-	FAssetTypeActions_InterchangeBlueprintPipelineBase(EAssetTypeCategories::Type InAssetCategory)
-		: AssetCategory(InAssetCategory)
-	{ }
-	// FAssetTypeActions_Blueprint interface
-	virtual UFactory* GetFactoryForBlueprintType(UBlueprint* InBlueprint) const override;
-
-	// IAssetTypeActions Implementation
-	virtual FText GetName() const override { return NSLOCTEXT("FAssetTypeActions_InterchangeBlueprintPipelineBase", "InterchangeBlueprintPipelineBaseName", "Interchange Blueprint Pipeline"); }
-	virtual FColor GetTypeColor() const override { return FColor(10, 25, 175); }
-	virtual UClass* GetSupportedClass() const override;
-	virtual uint32 GetCategories() override { return AssetCategory; }
-
-private:
-	EAssetTypeCategories::Type AssetCategory;
-};
-
 UCLASS(hidecategories = Object, collapsecategories)
 class UInterchangeEditorBlueprintPipelineBaseFactory : public UFactory
 {
@@ -100,26 +81,6 @@ public:
 	//End UFactory Interface	
 };
 
-class FAssetTypeActions_InterchangeEditorBlueprintPipelineBase : public FAssetTypeActions_Blueprint
-{
-public:
-	FAssetTypeActions_InterchangeEditorBlueprintPipelineBase(EAssetTypeCategories::Type InAssetCategory)
-		: AssetCategory(InAssetCategory)
-	{ }
-
-	// FAssetTypeActions_Blueprint interface
-	virtual UFactory* GetFactoryForBlueprintType(UBlueprint* InBlueprint) const override;
-
-	// IAssetTypeActions Implementation
-	virtual FText GetName() const override { return NSLOCTEXT("FAssetTypeActions_InterchangeEditorBlueprintPipelineBase", "InterchangeEditorBlueprintPipelineBaseName", "Interchange Editor Blueprint Pipeline"); }
-	virtual FColor GetTypeColor() const override { return FColor(10, 25, 175); }
-	virtual UClass* GetSupportedClass() const override;
-	virtual uint32 GetCategories() override { return AssetCategory; }
-
-private:
-	EAssetTypeCategories::Type AssetCategory;
-};
-
 UCLASS(hidecategories = Object, collapsecategories)
 class UInterchangePipelineBaseFactory : public UFactory
 {
@@ -137,23 +98,6 @@ private:
 	UClass* PipelineClass = nullptr;
 };
 
-class FAssetTypeActions_InterchangePipelineBase : public FAssetTypeActions_Base
-{
-public:
-	FAssetTypeActions_InterchangePipelineBase(EAssetTypeCategories::Type InAssetCategory)
-		: AssetCategory(InAssetCategory)
-	{ }
-
-	virtual FText GetName() const override { return NSLOCTEXT("FAssetTypeActions_InterchangePipelineBase", "InterchangePipelineBaseName", "Interchange Pipeline"); }
-	virtual FColor GetTypeColor() const override { return FColor(135, 200, 25); }
-	virtual UClass* GetSupportedClass() const override;
-	virtual uint32 GetCategories() override { return AssetCategory; }
-	virtual void OpenAssetEditor(const TArray<UObject*>& InObjects, TSharedPtr<class IToolkitHost> EditWithinLevelEditor = TSharedPtr<IToolkitHost>()) override;
-
-private:
-	EAssetTypeCategories::Type AssetCategory;
-};
-
 UCLASS(hidecategories = Object, collapsecategories)
 class UInterchangePythonPipelineAssetFactory : public UFactory
 {
@@ -170,21 +114,3 @@ public:
 private:
 	UClass* PythonClass = nullptr;
 };
-
-class FAssetTypeActions_InterchangePythonPipelineBase : public FAssetTypeActions_Base
-{
-public:
-	FAssetTypeActions_InterchangePythonPipelineBase(EAssetTypeCategories::Type InAssetCategory)
-		: AssetCategory(InAssetCategory)
-	{ }
-
-	// IAssetTypeActions Implementation
-	virtual FText GetName() const override { return NSLOCTEXT("InterchangePipelineFactories", "FAssetTypeActions_InterchangePythonPipeline", "Interchange Python Pipeline"); }
-	virtual FColor GetTypeColor() const override { return FColor(135, 200, 25); }
-	virtual UClass* GetSupportedClass() const override;
-	virtual uint32 GetCategories() override { return AssetCategory; }
-	virtual void OpenAssetEditor(const TArray<UObject*>& InObjects, TSharedPtr<class IToolkitHost> EditWithinLevelEditor = TSharedPtr<IToolkitHost>()) override;
-
-private:
-	EAssetTypeCategories::Type AssetCategory;
-};
diff --git a/Engine/Plugins/Interchange/Editor/Source/Pipelines/Public/InterchangeEditorPipelinesModule.h b/Engine/Plugins/Interchange/Editor/Source/Pipelines/Public/InterchangeEditorPipelinesModule.h
--- a/Engine/Plugins/Interchange/Editor/Source/Pipelines/Public/InterchangeEditorPipelinesModule.h
+++ b/Engine/Plugins/Interchange/Editor/Source/Pipelines/Public/InterchangeEditorPipelinesModule.h
@@ -4,11 +4,6 @@
 
 #include "Modules/ModuleManager.h"
 
-class FAssetTypeActions_InterchangeBlueprintPipelineBase;
-class FAssetTypeActions_InterchangeEditorBlueprintPipelineBase;
-class FAssetTypeActions_InterchangePipelineBase;
-class FAssetTypeActions_InterchangePythonPipelineBase;
-
 /**
  * The public interface to this module
  */
@@ -37,11 +32,4 @@ public:
 	{
 		return FModuleManager::Get().IsModuleLoaded( "InterchangeEditorPipelines" );
 	}
-
-	TArray<FName> ClassesToUnregisterOnShutdown;
-
-	TSharedPtr<FAssetTypeActions_InterchangeBlueprintPipelineBase> BlueprintPipelineBase_TypeActions;
-	TSharedPtr<FAssetTypeActions_InterchangeEditorBlueprintPipelineBase> BlueprintEditorPipelineBase_TypeActions;
-	TSharedPtr<FAssetTypeActions_InterchangePipelineBase> PipelineBase_TypeActions;
-	TSharedPtr<FAssetTypeActions_InterchangePythonPipelineBase> PythonPipelineBase_TypeActions;
 };
diff --git a/Engine/Plugins/Tests/InterchangeTests/Source/InterchangeTestEditor/Private/AssetDefinition_InterchangeImportTestPlan.cpp b/Engine/Plugins/Tests/InterchangeTests/Source/InterchangeTestEditor/Private/AssetDefinition_InterchangeImportTestPlan.cpp
--- a/Engine/Plugins/Tests/InterchangeTests/Source/InterchangeTestEditor/Private/AssetDefinition_InterchangeImportTestPlan.cpp
+++ b/Engine/Plugins/Tests/InterchangeTests/Source/InterchangeTestEditor/Private/AssetDefinition_InterchangeImportTestPlan.cpp
@@ -6,7 +6,10 @@
 
 TConstArrayView<FAssetCategoryPath> UAssetDefinition_InterchangeImportTestPlan::GetAssetCategories() const
 {
-	static const auto Categories = { FAssetCategoryPath(NSLOCTEXT("AssetDefinition_InterchangeImportTestPlan_Category", "Name", "Interchange Import Test Plan")) };
+	static const auto Categories = 
+		{
+			FAssetCategoryPath(EAssetCategoryPaths::Misc ,NSLOCTEXT("AssetDefinition", "InterchangeImportTestPlan_SubMenu", "Other"), ECategoryMenuType::Section)
+		};
 	return Categories;
 }
 
diff --git a/Engine/Plugins/Tests/InterchangeTests/Source/InterchangeTestEditor/Private/InterchangeImportTestPlanFactory.cpp b/Engine/Plugins/Tests/InterchangeTests/Source/InterchangeTestEditor/Private/InterchangeImportTestPlanFactory.cpp
--- a/Engine/Plugins/Tests/InterchangeTests/Source/InterchangeTestEditor/Private/InterchangeImportTestPlanFactory.cpp
+++ b/Engine/Plugins/Tests/InterchangeTests/Source/InterchangeTestEditor/Private/InterchangeImportTestPlanFactory.cpp
@@ -25,12 +25,6 @@ UObject* UInterchangeImportTestPlanFactory::FactoryCreateNew(UClass* Class, UObj
 }
 
 
-uint32 UInterchangeImportTestPlanFactory::GetMenuCategories() const
-{
-	return EAssetTypeCategories::Misc;
-}
-
-
 FText UInterchangeImportTestPlanFactory::GetDisplayName() const
 {
 	return LOCTEXT("MenuEntry", "Interchange Import Test Plan");
diff --git a/Engine/Plugins/Tests/InterchangeTests/Source/InterchangeTestEditor/Private/InterchangeImportTestPlanFactory.h b/Engine/Plugins/Tests/InterchangeTests/Source/InterchangeTestEditor/Private/InterchangeImportTestPlanFactory.h
--- a/Engine/Plugins/Tests/InterchangeTests/Source/InterchangeTestEditor/Private/InterchangeImportTestPlanFactory.h
+++ b/Engine/Plugins/Tests/InterchangeTests/Source/InterchangeTestEditor/Private/InterchangeImportTestPlanFactory.h
@@ -16,7 +16,6 @@ public:
 
 	// UFactory Interface
 	virtual UObject* FactoryCreateNew(UClass* Class, UObject* InParent, FName Name, EObjectFlags Flags, UObject* Context, FFeedbackContext* Warn) override;
-	virtual uint32 GetMenuCategories() const override;
 	virtual FText GetDisplayName() const override;
 	// End of UFactory Interface
 };
